[NO NON-PII COMMIT MESSAGE]

diff --git a/Engine/Plugins/Lumin/MagicLeap/Source/MagicLeapAudioCapture/Private/AudioCaptureMagicLeapModule.cpp b/Engine/Plugins/Lumin/MagicLeap/Source/MagicLeapAudioCapture/Private/AudioCaptureMagicLeapModule.cpp
--- a/Engine/Plugins/Lumin/MagicLeap/Source/MagicLeapAudioCapture/Private/AudioCaptureMagicLeapModule.cpp
+++ b/Engine/Plugins/Lumin/MagicLeap/Source/MagicLeapAudioCapture/Private/AudioCaptureMagicLeapModule.cpp
@@ -3,6 +3,7 @@
 #include "CoreMinimal.h"
 #include "AudioCaptureDeviceInterface.h"
 #include "AudioCaptureMagicLeap.h"
+#include "Modules/ModuleManager.h"
 #include "Features/IModularFeatures.h"
 
 namespace Audio
diff --git a/Engine/Source/Runtime/AudioCaptureImplementations/Android/AudioCaptureAndroid/Private/AudioCaptureAndroidModule.cpp b/Engine/Source/Runtime/AudioCaptureImplementations/Android/AudioCaptureAndroid/Private/AudioCaptureAndroidModule.cpp
--- a/Engine/Source/Runtime/AudioCaptureImplementations/Android/AudioCaptureAndroid/Private/AudioCaptureAndroidModule.cpp
+++ b/Engine/Source/Runtime/AudioCaptureImplementations/Android/AudioCaptureAndroid/Private/AudioCaptureAndroidModule.cpp
@@ -3,6 +3,7 @@
 #include "CoreMinimal.h"
 #include "AudioCaptureDeviceInterface.h"
 #include "AudioCaptureAndroid.h"
+#include "Modules/ModuleManager.h"
 #include "Features/IModularFeatures.h"
 
 namespace Audio
diff --git a/Engine/Source/Runtime/AudioCaptureImplementations/IOS/AudioCaptureAudioUnit/Private/AudioCaptureAudioUnitModule.cpp b/Engine/Source/Runtime/AudioCaptureImplementations/IOS/AudioCaptureAudioUnit/Private/AudioCaptureAudioUnitModule.cpp
--- a/Engine/Source/Runtime/AudioCaptureImplementations/IOS/AudioCaptureAudioUnit/Private/AudioCaptureAudioUnitModule.cpp
+++ b/Engine/Source/Runtime/AudioCaptureImplementations/IOS/AudioCaptureAudioUnit/Private/AudioCaptureAudioUnitModule.cpp
@@ -3,6 +3,7 @@
 #include "CoreMinimal.h"
 #include "AudioCaptureDeviceInterface.h"
 #include "AudioCaptureAudioUnit.h"
+#include "Modules/ModuleManager.h"
 #include "Features/IModularFeatures.h"
 
 namespace Audio
