[NO NON-PII COMMIT MESSAGE]

diff --git a/Engine/Source/Runtime/AudioCaptureImplementations/AudioCaptureRtAudio/AudioCaptureRtAudio.Build.cs b/Engine/Source/Runtime/AudioCaptureImplementations/AudioCaptureRtAudio/AudioCaptureRtAudio.Build.cs
--- a/Engine/Source/Runtime/AudioCaptureImplementations/AudioCaptureRtAudio/AudioCaptureRtAudio.Build.cs
+++ b/Engine/Source/Runtime/AudioCaptureImplementations/AudioCaptureRtAudio/AudioCaptureRtAudio.Build.cs
@@ -10,8 +10,7 @@ public class AudioCaptureRtAudio : ModuleRules
 		PrivateDependencyModuleNames.Add("AudioCaptureCore");
 
 		if (Target.Platform == UnrealTargetPlatform.Win32 ||
-			   Target.Platform == UnrealTargetPlatform.Win64 ||
-			   Target.Platform == UnrealTargetPlatform.XboxOne)
+			   Target.Platform == UnrealTargetPlatform.Win64)
 		{
 			PublicDefinitions.Add("WITH_RTAUDIO=1");
 			PublicDefinitions.Add("WITH_AUDIOCAPTURE=1");
@@ -19,6 +18,11 @@ public class AudioCaptureRtAudio : ModuleRules
 			// Allow us to use direct sound
 			AddEngineThirdPartyPrivateStaticDependencies(Target, "DirectSound");
 		}
+		else if(Target.Platform == UnrealTargetPlatform.XboxOne)
+        {
+            PublicDefinitions.Add("WITH_RTAUDIO=1");
+            PublicDefinitions.Add("WITH_AUDIOCAPTURE=1");
+        }
 		else if(Target.Platform == UnrealTargetPlatform.Mac)
         {
             PublicFrameworks.AddRange(new string[] { "CoreAudio", "AudioUnit", "AudioToolbox" });
