Allow FAudioCaptureAndroidStream to have CloseStream called on an already closed stream without asserting.

diff --git a/Engine/Source/Runtime/AudioCaptureImplementations/Android/AudioCaptureAndroid/Private/AudioCaptureAndroid.cpp b/Engine/Source/Runtime/AudioCaptureImplementations/Android/AudioCaptureAndroid/Private/AudioCaptureAndroid.cpp
--- a/Engine/Source/Runtime/AudioCaptureImplementations/Android/AudioCaptureAndroid/Private/AudioCaptureAndroid.cpp
+++ b/Engine/Source/Runtime/AudioCaptureImplementations/Android/AudioCaptureAndroid/Private/AudioCaptureAndroid.cpp
@@ -47,9 +47,11 @@ bool Audio::FAudioCaptureAndroidStream::OpenCaptureStream(const FAudioCaptureDev
 
 bool Audio::FAudioCaptureAndroidStream::CloseStream()
 {
-	check(InputOboeStream != nullptr);
-	InputOboeStream->close();
-	InputOboeStream.Reset();
+	if (InputOboeStream != nullptr)
+	{
+		InputOboeStream->close();
+		InputOboeStream.Reset();
+	}
 
 	return true;
 }
