minor fix up: use !InputOboeStream rather than InputOboeStream != nullptr for consistency with rest of file.

diff --git a/Engine/Source/Runtime/AudioCaptureImplementations/Android/AudioCaptureAndroid/Private/AudioCaptureAndroid.cpp b/Engine/Source/Runtime/AudioCaptureImplementations/Android/AudioCaptureAndroid/Private/AudioCaptureAndroid.cpp
--- a/Engine/Source/Runtime/AudioCaptureImplementations/Android/AudioCaptureAndroid/Private/AudioCaptureAndroid.cpp
+++ b/Engine/Source/Runtime/AudioCaptureImplementations/Android/AudioCaptureAndroid/Private/AudioCaptureAndroid.cpp
@@ -47,7 +47,7 @@ bool Audio::FAudioCaptureAndroidStream::OpenCaptureStream(const FAudioCaptureDev
 
 bool Audio::FAudioCaptureAndroidStream::CloseStream()
 {
-	if (InputOboeStream != nullptr)
+	if (!InputOboeStream)
 	{
 		InputOboeStream->close();
 		InputOboeStream.Reset();
