static anslysis warning fix

diff --git a/Engine/Source/Runtime/AudioCaptureImplementations/AudioCaptureRtAudio/Private/RtAudio.cpp b/Engine/Source/Runtime/AudioCaptureImplementations/AudioCaptureRtAudio/Private/RtAudio.cpp
--- a/Engine/Source/Runtime/AudioCaptureImplementations/AudioCaptureRtAudio/Private/RtAudio.cpp
+++ b/Engine/Source/Runtime/AudioCaptureImplementations/AudioCaptureRtAudio/Private/RtAudio.cpp
@@ -6205,6 +6205,9 @@ void RtApiDs :: abortStream()
   stopStream();
 }
 
+#pragma warning ( push )
+#pragma warning ( disable : 6385 ) // MSVC has issue with memcpy https://developercommunity.visualstudio.com/content/problem/841208/false-c6385-warning-when-using-memcpy.html
+
 void RtApiDs :: callbackEvent()
 {
   if ( stream_.state == STREAM_STOPPED || stream_.state == STREAM_STOPPING ) {
@@ -6618,6 +6621,7 @@ void RtApiDs :: callbackEvent()
   MUTEX_UNLOCK( &stream_.mutex );
   RtApi::tickStreamTime();
 }
+#pragma warning ( pop )
 
 // Definitions for utility functions and callbacks
 // specific to the DirectSound implementation.
