Port of audio capture implementation fixes

diff --git a/Engine/Plugins/Media/AudioCaptureTimecodeProvider/Source/AudioCaptureTimecodeProvider/Private/AudioCaptureTimecodeProvider.cpp b/Engine/Plugins/Media/AudioCaptureTimecodeProvider/Source/AudioCaptureTimecodeProvider/Private/AudioCaptureTimecodeProvider.cpp
--- a/Engine/Plugins/Media/AudioCaptureTimecodeProvider/Source/AudioCaptureTimecodeProvider/Private/AudioCaptureTimecodeProvider.cpp
+++ b/Engine/Plugins/Media/AudioCaptureTimecodeProvider/Source/AudioCaptureTimecodeProvider/Private/AudioCaptureTimecodeProvider.cpp
@@ -36,7 +36,7 @@ public:
 		//We want a fast timecode detection but we don't want to be called too often.
 		const int32 NumberCaptureFrames = 64;
 
-		Audio::FOnCaptureFunction OnCapture = [this](const float* AudioData, int32 NumFrames, int32 NumChannels, double StreamTime, bool bOverFlow)
+		Audio::FOnCaptureFunction OnCapture = [this](const float* AudioData, int32 NumFrames, int32 NumChannels, int32 SampleRate, double StreamTime, bool bOverFlow)
 		{
 			OnAudioCapture(AudioData, NumFrames, NumChannels, StreamTime, bOverFlow);
 		};
diff --git a/Engine/Plugins/Online/OnlineSubsystem/Source/OnlineSubsystem.Build.cs b/Engine/Plugins/Online/OnlineSubsystem/Source/OnlineSubsystem.Build.cs
--- a/Engine/Plugins/Online/OnlineSubsystem/Source/OnlineSubsystem.Build.cs
+++ b/Engine/Plugins/Online/OnlineSubsystem/Source/OnlineSubsystem.Build.cs
@@ -29,6 +29,7 @@ public class OnlineSubsystem : ModuleRules
 				"Sockets",
 				"JsonUtilities",
 				"AudioMixerCore",
+				"SignalProcessing",
 			}
 		);
 	}
diff --git a/Engine/Plugins/Runtime/AudioCapture/Source/AudioCapture/Private/AudioCapture.cpp b/Engine/Plugins/Runtime/AudioCapture/Source/AudioCapture/Private/AudioCapture.cpp
--- a/Engine/Plugins/Runtime/AudioCapture/Source/AudioCapture/Private/AudioCapture.cpp
+++ b/Engine/Plugins/Runtime/AudioCapture/Source/AudioCapture/Private/AudioCapture.cpp
@@ -23,7 +23,7 @@ bool UAudioCapture::OpenDefaultAudioStream()
 {
 	if (!AudioCapture.IsStreamOpen())
 	{
-		Audio::FOnCaptureFunction OnCapture = [this](const float* AudioData, int32 NumFrames, int32 InNumChannels, double StreamTime, bool bOverFlow)
+		Audio::FOnCaptureFunction OnCapture = [this](const float* AudioData, int32 NumFrames, int32 InNumChannels, int32 InSampleRate, double StreamTime, bool bOverFlow)
 		{
 			OnGeneratedAudio(AudioData, NumFrames * InNumChannels);
 		};
diff --git a/Engine/Source/Runtime/AudioCaptureCore/Private/AudioCaptureCore.cpp b/Engine/Source/Runtime/AudioCaptureCore/Private/AudioCaptureCore.cpp
--- a/Engine/Source/Runtime/AudioCaptureCore/Private/AudioCaptureCore.cpp
+++ b/Engine/Source/Runtime/AudioCaptureCore/Private/AudioCaptureCore.cpp
@@ -3,9 +3,8 @@
 #include "AudioCaptureCore.h"
 #include "AudioCaptureInternal.h"
 #include "CoreMinimal.h"
-#include "Modules/ModuleManager.h"
 
-DEFINE_LOG_CATEGORY(LogAudioCaptureCore)
+
 
 namespace Audio
 {
@@ -19,6 +18,26 @@ namespace Audio
 	{
 	}
 
+	bool FAudioCapture::RegisterUser(const TCHAR* UserId)
+	{
+		if (Impl.IsValid())
+		{
+			return Impl->RegisterUser(UserId);
+		}
+
+		return false;
+	}
+
+	bool FAudioCapture::UnregisterUser(const TCHAR* UserId)
+	{
+		if (Impl.IsValid())
+		{
+			return Impl->UnregisterUser(UserId);
+		}
+
+		return false;
+	}
+
 	int32 FAudioCapture::GetCaptureDevicesAvailable(TArray<FCaptureDeviceInfo>& OutDevices)
 	{
 		if (Impl.IsValid())
@@ -35,10 +54,7 @@ namespace Audio
 	{
 		if (Impl.IsValid())
 		{
-			if (DeviceIndex == DefaultDeviceIndex)
-			{
-				return Impl->GetCaptureDeviceInfo(OutInfo, DeviceIndex);
-			}
+			return Impl->GetCaptureDeviceInfo(OutInfo, DeviceIndex);
 		}
 
 		return false;
@@ -126,6 +142,23 @@ namespace Audio
 		return false;
 	}
 
+	bool FAudioCapture::GetIfHardwareFeatureIsSupported(EHardwareInputFeature FeatureType)
+	{
+		if (Impl.IsValid())
+		{
+			return Impl->GetIfHardwareFeatureIsSupported(FeatureType);
+		}
+		return false;
+	}
+
+	void FAudioCapture::SetHardwareFeatureEnabled(EHardwareInputFeature FeatureType, bool bIsEnabled)
+	{
+		if (Impl.IsValid())
+		{
+			Impl->SetHardwareFeatureEnabled(FeatureType, bIsEnabled);
+		}
+	}
+
 	FAudioCaptureSynth::FAudioCaptureSynth()
 		: NumSamplesEnqueued(0)
 		, bInitialized(false)
@@ -147,7 +180,7 @@ namespace Audio
 		bool bSuccess = true;
 		if (!AudioCapture.IsStreamOpen())
 		{
-			FOnCaptureFunction OnCapture = [this](const float* AudioData, int32 NumFrames, int32 NumChannels, double StreamTime, bool bOverFlow)
+			FOnCaptureFunction OnCapture = [this](const float* AudioData, int32 NumFrames, int32 NumChannels, int32 SampleRate, double StreamTime, bool bOverFlow)
 			{
 				int32 NumSamples = NumChannels * NumFrames;
 
@@ -242,15 +275,3 @@ namespace Audio
 	}
 
 } // namespace audio
-
-void FAudioCaptureCoreModule::StartupModule()
-{
-
-}
-
-void FAudioCaptureCoreModule::ShutdownModule()
-{
-
-}
-
-IMPLEMENT_MODULE(FAudioCaptureCoreModule, AudioCaptureCore);
\ No newline at end of file
diff --git a/Engine/Source/Runtime/AudioCaptureCore/Private/AudioCaptureCoreModule.cpp b/Engine/Source/Runtime/AudioCaptureCore/Private/AudioCaptureCoreModule.cpp
new file mode 100644
--- /dev/null
+++ b/Engine/Source/Runtime/AudioCaptureCore/Private/AudioCaptureCoreModule.cpp
@@ -0,0 +1,26 @@
+// Copyright Epic Games, Inc. All Rights Reserved.
+
+#include "AudioCaptureCoreModule.h"
+#include "AudioCaptureCoreLog.h"
+
+#include "Modules/ModuleManager.h"
+
+DEFINE_LOG_CATEGORY(LogAudioCaptureCore)
+
+
+void FAudioCaptureCoreModule::StartupModule()
+{
+
+}
+
+void FAudioCaptureCoreModule::ShutdownModule()
+{
+
+}
+
+
+#if defined(REQUIRE_EXPLICIT_GMALLOC_INIT) && REQUIRE_EXPLICIT_GMALLOC_INIT
+extern "C" void IMPLEMENT_MODULE_AudioCaptureCore() { }
+#else
+IMPLEMENT_MODULE(FAudioCaptureCoreModule, AudioCaptureCore);
+#endif
diff --git a/Engine/Source/Runtime/AudioCaptureCore/Private/AudioCaptureInternal.h b/Engine/Source/Runtime/AudioCaptureCore/Private/AudioCaptureInternal.h
--- a/Engine/Source/Runtime/AudioCaptureCore/Private/AudioCaptureInternal.h
+++ b/Engine/Source/Runtime/AudioCaptureCore/Private/AudioCaptureInternal.h
@@ -3,6 +3,7 @@
 #pragma once
 
 #include "AudioCaptureCore.h"
+#include "AudioCaptureCoreLog.h"
 #include "Features/IModularFeatures.h"
 
 namespace Audio
diff --git a/Engine/Source/Runtime/AudioCaptureCore/Public/AudioCaptureCore.h b/Engine/Source/Runtime/AudioCaptureCore/Public/AudioCaptureCore.h
--- a/Engine/Source/Runtime/AudioCaptureCore/Public/AudioCaptureCore.h
+++ b/Engine/Source/Runtime/AudioCaptureCore/Public/AudioCaptureCore.h
@@ -10,6 +10,14 @@
 
 namespace Audio
 {
+	// Various hardware-accelerated features that an input device can have.
+	enum class EHardwareInputFeature : uint8
+	{
+		EchoCancellation,
+		NoiseSuppression,
+		AutomaticGainControl
+	};
+
 	// Class which handles audio capture internally, implemented with a back-end per platform
 	class AUDIOCAPTURECORE_API FAudioCapture
 	{
@@ -21,6 +29,12 @@ namespace Audio
 		// Returns the total amount of audio devices.
 		int32 GetCaptureDevicesAvailable(TArray<FCaptureDeviceInfo>& OutDevices);
 
+		// Adds a user to the system so we can use its devices.
+		bool RegisterUser(const TCHAR* UserId);
+
+		// Removes a user added with RegisterUser.
+		bool UnregisterUser(const TCHAR* UserId);
+
 		// Returns the audio capture device information at the given Id.
 		bool GetCaptureDeviceInfo(FCaptureDeviceInfo& OutInfo, int32 DeviceIndex = INDEX_NONE);
 
@@ -51,6 +65,10 @@ namespace Audio
 		// Returns true if the audio capture stream is currently capturing audio
 		bool IsCapturing() const;
 
+		bool GetIfHardwareFeatureIsSupported(EHardwareInputFeature FeatureType);
+
+		void SetHardwareFeatureEnabled(EHardwareInputFeature FeatureType, bool bIsEnabled);
+
 	private:
 
 		TUniquePtr<IAudioCaptureStream> CreateImpl();
@@ -118,12 +136,3 @@ namespace Audio
 	};
 
 } // namespace Audio
-
-DECLARE_LOG_CATEGORY_EXTERN(LogAudioCaptureCore, Log, All);
-
-class FAudioCaptureCoreModule : public IModuleInterface
-{
-public:
-	virtual void StartupModule() override;
-	virtual void ShutdownModule() override;
-};
\ No newline at end of file
diff --git a/Engine/Source/Runtime/AudioCaptureCore/Public/AudioCaptureCoreLog.h b/Engine/Source/Runtime/AudioCaptureCore/Public/AudioCaptureCoreLog.h
new file mode 100644
--- /dev/null
+++ b/Engine/Source/Runtime/AudioCaptureCore/Public/AudioCaptureCoreLog.h
@@ -0,0 +1,7 @@
+// Copyright Epic Games, Inc. All Rights Reserved.
+
+#pragma once
+
+#include "CoreMinimal.h"
+
+DECLARE_LOG_CATEGORY_EXTERN(LogAudioCaptureCore, Log, All);
diff --git a/Engine/Source/Runtime/AudioCaptureCore/Public/AudioCaptureCoreModule.h b/Engine/Source/Runtime/AudioCaptureCore/Public/AudioCaptureCoreModule.h
new file mode 100644
--- /dev/null
+++ b/Engine/Source/Runtime/AudioCaptureCore/Public/AudioCaptureCoreModule.h
@@ -0,0 +1,17 @@
+// Copyright Epic Games, Inc. All Rights Reserved.
+
+#pragma once
+
+
+/* Public dependencies
+*****************************************************************************/
+
+#include "CoreMinimal.h"
+#include "Modules/ModuleManager.h"
+
+class FAudioCaptureCoreModule : public IModuleInterface
+{
+public:
+	virtual void StartupModule() override;
+	virtual void ShutdownModule() override;
+};
diff --git a/Engine/Source/Runtime/AudioCaptureCore/Public/AudioCaptureDeviceInterface.h b/Engine/Source/Runtime/AudioCaptureCore/Public/AudioCaptureDeviceInterface.h
--- a/Engine/Source/Runtime/AudioCaptureCore/Public/AudioCaptureDeviceInterface.h
+++ b/Engine/Source/Runtime/AudioCaptureCore/Public/AudioCaptureDeviceInterface.h
@@ -7,9 +7,12 @@
 
 namespace Audio
 {
+	enum class EHardwareInputFeature : uint8;
+	
 	struct FCaptureDeviceInfo
  	{
 		FString DeviceName;
+		FString DeviceId;
 		int32 InputChannels;
 		int32 PreferredSampleRate;
 		bool bSupportsHardwareAEC;
@@ -31,7 +34,7 @@ namespace Audio
 		}
 	};
 
-	typedef TFunction<void(const float* InAudio, int32 NumFrames, int32 NumChannels, double StreamTime, bool bOverFlow)> FOnCaptureFunction;
+	typedef TFunction<void(const float* InAudio, int32 NumFrames, int32 NumChannels, int32 SampleRate, double StreamTime, bool bOverFlow)> FOnCaptureFunction;
 	
 	class IAudioCaptureStream : public IModularFeature
 	{
@@ -40,7 +43,10 @@ namespace Audio
 		virtual ~IAudioCaptureStream() {}
 
 		
-
+		// Lets us know which users are in the system.
+		virtual bool RegisterUser(const TCHAR* UserId) { return true; }
+		// Call this to remove a user that was added with RegisterUser.
+		virtual bool UnregisterUser(const TCHAR* UserId) { return true; }
 		// Returns the audio capture device information at the given Id.
 		virtual bool GetCaptureDeviceInfo(FCaptureDeviceInfo& OutInfo, int32 DeviceIndex) = 0;
 		// Opens the audio capture stream with the given parameters
@@ -67,6 +73,15 @@ namespace Audio
 
 		// Returns the total amount of audio devices.
 		virtual bool GetInputDevicesAvailable(TArray<FCaptureDeviceInfo>& OutDevices) = 0;
+
+		virtual bool GetIfHardwareFeatureIsSupported(EHardwareInputFeature FeatureType) 
+		{ 
+			return false;
+		}
+
+		virtual void SetHardwareFeatureEnabled(EHardwareInputFeature FeatureType, bool bEnabled)
+		{
+		}
 	};
 
 	class IAudioCaptureFactory : public IModularFeature
diff --git a/Engine/Source/Runtime/AudioCaptureImplementations/Android/AudioCaptureAndroid/Private/AudioCaptureAndroid.cpp b/Engine/Source/Runtime/AudioCaptureImplementations/Android/AudioCaptureAndroid/Private/AudioCaptureAndroid.cpp
--- a/Engine/Source/Runtime/AudioCaptureImplementations/Android/AudioCaptureAndroid/Private/AudioCaptureAndroid.cpp
+++ b/Engine/Source/Runtime/AudioCaptureImplementations/Android/AudioCaptureAndroid/Private/AudioCaptureAndroid.cpp
@@ -115,7 +115,7 @@ bool Audio::FAudioCaptureAndroidStream::IsCapturing() const
 void Audio::FAudioCaptureAndroidStream::OnAudioCapture(void* InBuffer, uint32 InBufferFrames, double StreamTime, bool bOverflow)
 {
 	const float* FloatBuffer = static_cast<float*>(InBuffer);
-	OnCapture(FloatBuffer, InBufferFrames, NumChannels, StreamTime, bOverflow);
+	OnCapture(FloatBuffer, InBufferFrames, NumChannels, SampleRate, StreamTime, bOverflow);
 }
 
 bool Audio::FAudioCaptureAndroidStream::GetInputDevicesAvailable(TArray<FCaptureDeviceInfo>& OutDevices)
diff --git a/Engine/Source/Runtime/AudioCaptureImplementations/AudioCaptureRtAudio/Private/AudioCaptureRtAudio.cpp b/Engine/Source/Runtime/AudioCaptureImplementations/AudioCaptureRtAudio/Private/AudioCaptureRtAudio.cpp
--- a/Engine/Source/Runtime/AudioCaptureImplementations/AudioCaptureRtAudio/Private/AudioCaptureRtAudio.cpp
+++ b/Engine/Source/Runtime/AudioCaptureImplementations/AudioCaptureRtAudio/Private/AudioCaptureRtAudio.cpp
@@ -31,9 +31,8 @@ bool Audio::FAudioCaptureRtAudioStream::GetCaptureDeviceInfo(FCaptureDeviceInfo&
 
 	RtAudio::DeviceInfo DeviceInfo = CaptureDevice.getDeviceInfo(InputDeviceId);
 
-	const char* NameStr = DeviceInfo.name.c_str();
-
-	OutInfo.DeviceName = FString(ANSI_TO_TCHAR(NameStr));
+	OutInfo.DeviceName = FString(ANSI_TO_TCHAR(DeviceInfo.name.c_str()));
+	OutInfo.DeviceId = FString(ANSI_TO_TCHAR(DeviceInfo.deviceId.c_str()));
 	OutInfo.InputChannels = DeviceInfo.inputChannels;
 	OutInfo.PreferredSampleRate = DeviceInfo.preferredSampleRate;
 
@@ -49,7 +48,32 @@ bool Audio::FAudioCaptureRtAudioStream::OpenCaptureStream(const FAudioCaptureDev
 	uint32 InputDeviceId = CaptureDevice.getDefaultInputDevice();
 	if (InParams.DeviceIndex != DefaultDeviceIndex)
 	{
-		InputDeviceId = InParams.DeviceIndex - 1;
+		// InParams.DeviceIndex is the index inside the sublist of devices including only
+		// the Input devices.
+		uint32 NumInputDevices = 0;
+		bool DeviceFound = false;
+		uint32 NumDevices = CaptureDevice.getDeviceCount();
+		for (uint32 DeviceIndex = 0; DeviceIndex < NumDevices; DeviceIndex++)
+		{
+			FCaptureDeviceInfo DeviceInfo;
+			GetCaptureDeviceInfo(DeviceInfo, DeviceIndex);
+			bool IsInput = DeviceInfo.InputChannels > 0;
+			if (!IsInput)
+			{
+				continue;
+			}
+			if (NumInputDevices == InParams.DeviceIndex)
+			{
+				DeviceFound = true;
+				InputDeviceId = DeviceIndex;
+				break;
+			}
+			NumInputDevices++;
+		}
+		if (!DeviceFound)
+		{
+			return false;
+		}
 	}
 
 	RtAudio::DeviceInfo DeviceInfo = CaptureDevice.getDeviceInfo(InputDeviceId);
@@ -167,7 +191,7 @@ void Audio::FAudioCaptureRtAudioStream::OnAudioCapture(void* InBuffer, uint32 In
 {
 #if WITH_RTAUDIO
 	float* InBufferData = (float*)InBuffer;
-	OnCapture(InBufferData, InBufferFrames, NumChannels, StreamTime, bOverflow);
+	OnCapture(InBufferData, InBufferFrames, NumChannels, SampleRate, StreamTime, bOverflow);
 #endif // WITH_RTAUDIO
 }
 
@@ -176,10 +200,16 @@ bool Audio::FAudioCaptureRtAudioStream::GetInputDevicesAvailable(TArray<FCapture
 #if WITH_RTAUDIO
 	uint32 NumDevices = CaptureDevice.getDeviceCount();
 	OutDevices.Reset();
-	for (uint32 DeviceIndex = 1; DeviceIndex < NumDevices; DeviceIndex++)
+	// Iterate over all devices and include in the output only the ones that are input devices
+	for (uint32 DeviceIndex = 0; DeviceIndex < NumDevices; DeviceIndex++)
 	{
-		FCaptureDeviceInfo& DeviceInfo = OutDevices.AddDefaulted_GetRef();
+		FCaptureDeviceInfo DeviceInfo;
 		GetCaptureDeviceInfo(DeviceInfo, DeviceIndex);
+		bool IsInput = DeviceInfo.InputChannels > 0;
+		if (IsInput)
+		{
+			OutDevices.Add(DeviceInfo);
+		}
 	}
 
 	return true;
diff --git a/Engine/Source/Runtime/AudioCaptureImplementations/AudioCaptureRtAudio/Private/RtAudio.cpp b/Engine/Source/Runtime/AudioCaptureImplementations/AudioCaptureRtAudio/Private/RtAudio.cpp
--- a/Engine/Source/Runtime/AudioCaptureImplementations/AudioCaptureRtAudio/Private/RtAudio.cpp
+++ b/Engine/Source/Runtime/AudioCaptureImplementations/AudioCaptureRtAudio/Private/RtAudio.cpp
@@ -100,6 +100,25 @@ static void throw_wrapper(const RtAudioError& Error)
     return s;
   }
 
+  static std::string convertGuidToStdString(LPGUID guid)
+  {
+	char result[39]; // 32 hex chars + 4 hyphens + 2 braces + null
+	if (guid)
+	{
+		snprintf(
+			result, sizeof(result), "{%08lx-%04x-%04x-%02x%02x-%02x%02x%02x%02x%02x%02x}",
+			guid->Data1, guid->Data2, guid->Data3,
+			guid->Data4[0], guid->Data4[1], guid->Data4[2], guid->Data4[3],
+			guid->Data4[4], guid->Data4[5], guid->Data4[6], guid->Data4[7]);
+	}
+	else
+	{
+		result[0] = '\0';
+	}
+
+	return result;
+  }
+
 #elif defined(__LINUX_ALSA__) || defined(__LINUX_PULSE__) || defined(__UNIX_JACK__) || defined(__LINUX_OSS__) || defined(__MACOSX_CORE__)
   // pthread API
   #define MUTEX_INITIALIZE(A) pthread_mutex_init(A, NULL)
@@ -664,19 +683,6 @@ RtAudio::DeviceInfo RtApiCore :: getDeviceInfo( unsigned int device )
     return info;
   }
 
-  //const char *mname = CFStringGetCStringPtr( cfname, CFStringGetSystemEncoding() );
-  int length = CFStringGetLength(cfname);
-  char *mname = (char *)malloc(length * 3 + 1);
-#if defined( UNICODE ) || defined( _UNICODE )
-  CFStringGetCString(cfname, mname, length * 3 + 1, kCFStringEncodingUTF8);
-#else
-  CFStringGetCString(cfname, mname, length * 3 + 1, CFStringGetSystemEncoding());
-#endif
-  info.name.append( (const char *)mname, strlen(mname) );
-  info.name.append( ": " );
-  CFRelease( cfname );
-  free(mname);
-
   property.mSelector = kAudioObjectPropertyName;
   result = AudioObjectGetPropertyData( id, &property, 0, NULL, &dataSize, &cfname );
   if ( result != noErr ) {
@@ -687,7 +693,7 @@ RtAudio::DeviceInfo RtApiCore :: getDeviceInfo( unsigned int device )
   }
 
   //const char *name = CFStringGetCStringPtr( cfname, CFStringGetSystemEncoding() );
-  length = CFStringGetLength(cfname);
+  int length = CFStringGetLength(cfname);
   char *name = (char *)malloc(length * 3 + 1);
 #if defined( UNICODE ) || defined( _UNICODE )
   CFStringGetCString(cfname, name, length * 3 + 1, kCFStringEncodingUTF8);
@@ -4009,11 +4015,12 @@ RtAudio::DeviceInfo RtApiWasapi::getDeviceInfo( unsigned int device )
   RtAudio::DeviceInfo info;
   unsigned int captureDeviceCount = 0;
   unsigned int renderDeviceCount = 0;
-  std::string defaultDeviceName;
   bool isCaptureDevice = false;
+  std::string defaultDeviceId;
+  LPWSTR strDefaultDeviceId;
+  LPWSTR strDeviceId;
 
   PROPVARIANT deviceNameProp;
-  PROPVARIANT defaultDeviceNameProp;
 
   IMMDeviceCollection* captureDevices = NULL;
   IMMDeviceCollection* renderDevices = NULL;
@@ -4098,22 +4105,27 @@ RtAudio::DeviceInfo RtApiWasapi::getDeviceInfo( unsigned int device )
     }
   }
 
+  hr = defaultDevicePtr->GetId(&strDefaultDeviceId);
+  if (FAILED(hr)) {
+	  errorText_ = "RtApiWasapi::getDeviceInfo: Unable to get default device Id.";
+	  goto Exit;
+  }
+  defaultDeviceId = convertCharPointerToStdString(strDefaultDeviceId);
+
+  hr = devicePtr->GetId(&strDeviceId);
+  if (FAILED(hr)) {
+	  errorText_ = "RtApiWasapi::getDeviceInfo: Unable to get device Id.";
+	  goto Exit;
+  }
+  info.deviceId = convertCharPointerToStdString(strDeviceId);
+
 #if PLATFORM_WINDOWS
   hr = defaultDevicePtr->OpenPropertyStore( STGM_READ, &defaultDevicePropStore );
   if ( FAILED( hr ) ) {
     errorText_ = "RtApiWasapi::getDeviceInfo: Unable to open default device property store.";
     goto Exit;
   }
-  PropVariantInit( &defaultDeviceNameProp );
-
-  hr = defaultDevicePropStore->GetValue( PKEY_Device_FriendlyName, &defaultDeviceNameProp );
-  if ( FAILED( hr ) ) {
-    errorText_ = "RtApiWasapi::getDeviceInfo: Unable to retrieve default device property: PKEY_Device_FriendlyName.";
-    goto Exit;
-  }
-
-  defaultDeviceName = convertCharPointerToStdString(defaultDeviceNameProp.pwszVal);
-
+ 
   // name
   hr = devicePtr->OpenPropertyStore( STGM_READ, &devicePropStore );
   if ( FAILED( hr ) ) {
@@ -4129,20 +4141,19 @@ RtAudio::DeviceInfo RtApiWasapi::getDeviceInfo( unsigned int device )
     goto Exit;
   }
 
-  info.name =convertCharPointerToStdString(deviceNameProp.pwszVal);
+  info.name = convertCharPointerToStdString(deviceNameProp.pwszVal);
 #else // On Xbox One, we hardcode out defaults.
-  defaultDeviceName = "Xbox One Input";
-  info.name = "Xbox One Input";
-#endif // PLATFORM_WINDOWS 
+  info.name = std::string("Xbox Audio Input ");
+#endif
 
   // is default
   if ( isCaptureDevice ) {
-    info.isDefaultInput = info.name == defaultDeviceName;
+    info.isDefaultInput = info.deviceId == defaultDeviceId;
     info.isDefaultOutput = false;
   }
   else {
     info.isDefaultInput = false;
-    info.isDefaultOutput = info.name == defaultDeviceName;
+    info.isDefaultOutput = info.deviceId == defaultDeviceId;;
   }
 
   // channel count
@@ -4213,7 +4224,6 @@ RtAudio::DeviceInfo RtApiWasapi::getDeviceInfo( unsigned int device )
 Exit:
   // release all references
   PropVariantClear( &deviceNameProp );
-  PropVariantClear( &defaultDeviceNameProp );
 
   SAFE_RELEASE( captureDevices );
   SAFE_RELEASE( renderDevices );
@@ -5373,6 +5383,7 @@ RtAudio::DeviceInfo RtApiDs :: getDeviceInfo( unsigned int device )
 
   if ( dsDevices[ device ].validId[1] == false ) {
     info.name = dsDevices[ device ].name;
+	info.deviceId = convertGuidToStdString(dsDevices[device].id[0]);
     info.probed = true;
     return info;
   }
@@ -5412,18 +5423,21 @@ RtAudio::DeviceInfo RtApiDs :: getDeviceInfo( unsigned int device )
     if ( inCaps.dwFormats & WAVE_FORMAT_1S08 ) info.nativeFormats |= RTAUDIO_SINT8;
     if ( inCaps.dwFormats & WAVE_FORMAT_2S08 ) info.nativeFormats |= RTAUDIO_SINT8;
     if ( inCaps.dwFormats & WAVE_FORMAT_4S08 ) info.nativeFormats |= RTAUDIO_SINT8;
+	if ( inCaps.dwFormats & WAVE_FORMAT_48S08 ) info.nativeFormats |= RTAUDIO_SINT8;
     if ( inCaps.dwFormats & WAVE_FORMAT_96S08 ) info.nativeFormats |= RTAUDIO_SINT8;
 
     if ( info.nativeFormats & RTAUDIO_SINT16 ) {
       if ( inCaps.dwFormats & WAVE_FORMAT_1S16 ) rates.push_back( 11025 );
       if ( inCaps.dwFormats & WAVE_FORMAT_2S16 ) rates.push_back( 22050 );
       if ( inCaps.dwFormats & WAVE_FORMAT_4S16 ) rates.push_back( 44100 );
+	  if ( inCaps.dwFormats & WAVE_FORMAT_48S16 ) rates.push_back( 48000 );
       if ( inCaps.dwFormats & WAVE_FORMAT_96S16 ) rates.push_back( 96000 );
     }
     else if ( info.nativeFormats & RTAUDIO_SINT8 ) {
       if ( inCaps.dwFormats & WAVE_FORMAT_1S08 ) rates.push_back( 11025 );
       if ( inCaps.dwFormats & WAVE_FORMAT_2S08 ) rates.push_back( 22050 );
       if ( inCaps.dwFormats & WAVE_FORMAT_4S08 ) rates.push_back( 44100 );
+	  if ( inCaps.dwFormats & WAVE_FORMAT_48S08 ) rates.push_back( 48000 );
       if ( inCaps.dwFormats & WAVE_FORMAT_96S08 ) rates.push_back( 96000 );
     }
   }
@@ -5435,18 +5449,21 @@ RtAudio::DeviceInfo RtApiDs :: getDeviceInfo( unsigned int device )
     if ( inCaps.dwFormats & WAVE_FORMAT_1M08 ) info.nativeFormats |= RTAUDIO_SINT8;
     if ( inCaps.dwFormats & WAVE_FORMAT_2M08 ) info.nativeFormats |= RTAUDIO_SINT8;
     if ( inCaps.dwFormats & WAVE_FORMAT_4M08 ) info.nativeFormats |= RTAUDIO_SINT8;
+	if ( inCaps.dwFormats & WAVE_FORMAT_48M08 ) info.nativeFormats |= RTAUDIO_SINT8;
     if ( inCaps.dwFormats & WAVE_FORMAT_96M08 ) info.nativeFormats |= RTAUDIO_SINT8;
 
     if ( info.nativeFormats & RTAUDIO_SINT16 ) {
       if ( inCaps.dwFormats & WAVE_FORMAT_1M16 ) rates.push_back( 11025 );
       if ( inCaps.dwFormats & WAVE_FORMAT_2M16 ) rates.push_back( 22050 );
       if ( inCaps.dwFormats & WAVE_FORMAT_4M16 ) rates.push_back( 44100 );
+	  if ( inCaps.dwFormats & WAVE_FORMAT_48M16 ) rates.push_back( 48000 );
       if ( inCaps.dwFormats & WAVE_FORMAT_96M16 ) rates.push_back( 96000 );
     }
     else if ( info.nativeFormats & RTAUDIO_SINT8 ) {
       if ( inCaps.dwFormats & WAVE_FORMAT_1M08 ) rates.push_back( 11025 );
       if ( inCaps.dwFormats & WAVE_FORMAT_2M08 ) rates.push_back( 22050 );
       if ( inCaps.dwFormats & WAVE_FORMAT_4M08 ) rates.push_back( 44100 );
+	  if ( inCaps.dwFormats & WAVE_FORMAT_48M08) rates.push_back( 48000 );
       if ( inCaps.dwFormats & WAVE_FORMAT_96M08 ) rates.push_back( 96000 );
     }
   }
@@ -5469,6 +5486,13 @@ RtAudio::DeviceInfo RtApiDs :: getDeviceInfo( unsigned int device )
     if ( found == false ) info.sampleRates.push_back( rates[i] );
   }
   std::sort( info.sampleRates.begin(), info.sampleRates.end() );
+  
+  // Get preferred sample rate
+  for (unsigned int i = 0; i < info.sampleRates.size(); i++) {
+	  if (info.sampleRates[i] > info.preferredSampleRate && info.preferredSampleRate < 48000) {
+		  info.preferredSampleRate = info.sampleRates[i];
+	  }
+  }
 
   // If device opens for both playback and capture, we determine the channels.
   if ( info.outputChannels > 0 && info.inputChannels > 0 )
@@ -5478,6 +5502,7 @@ RtAudio::DeviceInfo RtApiDs :: getDeviceInfo( unsigned int device )
 
   // Copy name and return.
   info.name = dsDevices[ device ].name;
+  info.deviceId = convertGuidToStdString(dsDevices[device].id[1]);
   info.probed = true;
   return info;
 }
diff --git a/Engine/Source/Runtime/AudioCaptureImplementations/AudioCaptureRtAudio/Private/RtAudio.h b/Engine/Source/Runtime/AudioCaptureImplementations/AudioCaptureRtAudio/Private/RtAudio.h
--- a/Engine/Source/Runtime/AudioCaptureImplementations/AudioCaptureRtAudio/Private/RtAudio.h
+++ b/Engine/Source/Runtime/AudioCaptureImplementations/AudioCaptureRtAudio/Private/RtAudio.h
@@ -330,7 +330,8 @@ class AUDIOCAPTURERTAUDIO_API RtAudio
   //! The public device information structure for returning queried values.
   struct DeviceInfo {
     bool probed;                  /*!< true if the device capabilities were successfully probed. */
-    std::string name;             /*!< Character string device identifier. */
+	std::string name;             /*!< Character string device display name. */
+	std::string deviceId;               /*!< Character string device identifier. */
     unsigned int outputChannels;  /*!< Maximum output channels supported by device. */
     unsigned int inputChannels;   /*!< Maximum input channels supported by device. */
     unsigned int duplexChannels;  /*!< Maximum simultaneous input/output channels supported by device. */
diff --git a/Engine/Source/Runtime/AudioCaptureImplementations/IOS/AudioCaptureAudioUnit/Private/AudioCaptureAudioUnit.cpp b/Engine/Source/Runtime/AudioCaptureImplementations/IOS/AudioCaptureAudioUnit/Private/AudioCaptureAudioUnit.cpp
--- a/Engine/Source/Runtime/AudioCaptureImplementations/IOS/AudioCaptureAudioUnit/Private/AudioCaptureAudioUnit.cpp
+++ b/Engine/Source/Runtime/AudioCaptureImplementations/IOS/AudioCaptureAudioUnit/Private/AudioCaptureAudioUnit.cpp
@@ -2,6 +2,8 @@
 
 #include "AudioCaptureAudioUnit.h"
 
+const int32 kInputBus = 1;
+const int32 kOutputBus = 0;
 
 Audio::FAudioCaptureAudioUnitStream::FAudioCaptureAudioUnitStream()
 	: NumChannels(0)
@@ -16,24 +18,60 @@ static OSStatus RecordingCallback(void *inRefCon,
 	UInt32 inNumberFrames,
 	AudioBufferList *ioData)
 {
-	OSStatus status = noErr;
 	Audio::FAudioCaptureAudioUnitStream* AudioCapture = (Audio::FAudioCaptureAudioUnitStream*)inRefCon;
+	return AudioCapture->OnCaptureCallback(ioActionFlags, inTimeStamp, inBusNumber, inNumberFrames, ioData);
+}
+
+OSStatus Audio::FAudioCaptureAudioUnitStream::OnCaptureCallback(AudioUnitRenderActionFlags* ioActionFlags, const AudioTimeStamp* inTimeStamp, UInt32 inBusNumber, UInt32 inNumberFrames, AudioBufferList* ioData)
+{
+	OSStatus status = noErr;
 
-	status = AudioUnitRender(AudioCapture->AudioUnit,
+	const int NeededBufferSize = inNumberFrames * NumChannels * sizeof(float);
+	if (CaptureBuffer.Num() == 0 || BufferSize < NeededBufferSize)
+	{
+		BufferSize = NeededBufferSize;
+		AllocateBuffer(BufferSize);
+	}
+	
+	AudioBufferList* BufferList = (AudioBufferList*) CaptureBuffer.GetData();
+	for (int i=0; i < BufferList->mNumberBuffers; ++i) {
+		BufferList->mBuffers[i].mDataByteSize = (UInt32) BufferSize;
+	}
+	
+	status = AudioUnitRender(IOUnit,
 		ioActionFlags,
 		inTimeStamp,
 		inBusNumber,
 		inNumberFrames,
-		ioData);
+		BufferList);
 	check(status == noErr);
 
-
-	void* InBuffer = (void*)ioData->mBuffers[0].mData;
-	AudioCapture->OnAudioCapture(InBuffer, inNumberFrames, 0.0, false);
+	void* InBuffer = (void*)BufferList->mBuffers[0].mData; // only first channel ?!
+	OnAudioCapture(InBuffer, inNumberFrames, 0.0, false); // need calculate timestamp
 
 	return noErr;
 }
 
+void Audio::FAudioCaptureAudioUnitStream::AllocateBuffer(int SizeInBytes)
+{
+	size_t NeedBytes = sizeof(AudioBufferList) + NumChannels * (sizeof(AudioBuffer) + SizeInBytes);
+
+	CaptureBuffer.SetNum(NeedBytes);
+
+	AudioBufferList* list = (AudioBufferList*) CaptureBuffer.GetData();
+	uint8* data = CaptureBuffer.GetData() + sizeof(AudioBufferList) + sizeof(AudioBuffer);
+	
+	list->mNumberBuffers = NumChannels;
+	for(int i=0; i < NumChannels; i++)
+	{
+		list->mBuffers[i].mNumberChannels = 1;
+		list->mBuffers[i].mDataByteSize   = (UInt32) SizeInBytes;
+		list->mBuffers[i].mData           = data;
+		
+		data += (SizeInBytes + sizeof(AudioBuffer));
+	}
+}
+
 bool Audio::FAudioCaptureAudioUnitStream::GetCaptureDeviceInfo(FCaptureDeviceInfo& OutInfo, int32 DeviceIndex)
 {
 	OutInfo.DeviceName = FString(TEXT("Remote IO Audio Component"));
@@ -45,97 +83,107 @@ bool Audio::FAudioCaptureAudioUnitStream::GetCaptureDeviceInfo(FCaptureDeviceInf
 
 bool Audio::FAudioCaptureAudioUnitStream::OpenCaptureStream(const FAudioCaptureDeviceParams& InParams, FOnCaptureFunction InOnCapture, uint32 NumFramesDesired)
 {
+	NumChannels = 1;
+	SampleRate = 48000;
+	OSStatus Status = noErr;
+	
+	OnCapture = MoveTemp(InOnCapture);
+	
+	// Source of info "Technical Note TN2091 - Device input using the HAL Output Audio Unit"
+	
 	AudioComponentDescription desc;
 	desc.componentType = kAudioUnitType_Output;
-	if (InParams.bUseHardwareAEC)
-	{
-		desc.componentSubType = kAudioUnitSubType_RemoteIO;
-	}
-	else
-	{
-		desc.componentSubType = kAudioUnitSubType_RemoteIO;
-	}
-
+	// We use processing element always for enable runtime changing HW AEC and AGC settings.
+	// When it's disable the unit work as RemoteIO
+	desc.componentSubType = kAudioUnitSubType_VoiceProcessingIO;
 	desc.componentManufacturer = kAudioUnitManufacturer_Apple;
 	desc.componentFlags = 0;
 	desc.componentFlagsMask = 0;
 
 	AudioComponent InputComponent = AudioComponentFindNext(NULL, &desc);
 
-	AudioUnitStatus = AudioComponentInstanceNew(InputComponent, &AudioUnit);
-	check(AudioUnitStatus == noErr);
+	Status = AudioComponentInstanceNew(InputComponent, &IOUnit);
+	check(Status == noErr);
 
-	// Enable IO via AudioUnitSetProperty:
-	static const uint32 AudioUnitTrue = 1;
-	AudioUnitStatus = AudioUnitSetProperty(AudioUnit,
+	// Enable recording
+	uint32 EnableIO = 1;
+	Status = AudioUnitSetProperty(IOUnit,
 		kAudioOutputUnitProperty_EnableIO,
 		kAudioUnitScope_Input,
-		1,
-		&AudioUnitTrue,
-		sizeof(AudioUnitTrue));
-	check(AudioUnitStatus == noErr);
-
-	AudioUnitStatus = AudioUnitSetProperty(AudioUnit,
+		kInputBus,
+		&EnableIO,
+		sizeof(EnableIO));
+	check(Status == noErr);
+
+	// Disable output part
+	EnableIO = 0;
+	Status = AudioUnitSetProperty(IOUnit,
 		kAudioOutputUnitProperty_EnableIO,
 		kAudioUnitScope_Output,
-		0,
-		&AudioUnitTrue,
-		sizeof(AudioUnitTrue));
-	check(AudioUnitStatus == noErr);
-
-	AudioStreamBasicDescription StreamDescription;
-
-	StreamDescription.mSampleRate = 48000.00;
-	StreamDescription.mFormatID = kAudioFormatLinearPCM;
-	StreamDescription.mFormatFlags = kAudioFormatFlagIsFloat | kAudioFormatFlagsNativeEndian | kAudioFormatFlagIsPacked;
-	StreamDescription.mFramesPerPacket = NumFramesDesired;
-	StreamDescription.mChannelsPerFrame = 1;
-	StreamDescription.mBitsPerChannel = 8 * sizeof(float);
-	StreamDescription.mBytesPerPacket = NumFramesDesired * sizeof(float);
-	StreamDescription.mBytesPerFrame = sizeof(float);
-
-	AudioUnitStatus = AudioUnitSetProperty(AudioUnit,
+		kOutputBus,
+		&EnableIO,
+		sizeof(EnableIO));
+	check(Status == noErr);
+
+	AudioStreamBasicDescription StreamDescription = {0};
+	const UInt32 BytesPerSample = sizeof(Float32);
+	
+	StreamDescription.mSampleRate       = SampleRate;
+	StreamDescription.mFormatID         = kAudioFormatLinearPCM;
+	StreamDescription.mFormatFlags      = kAudioFormatFlagsNativeFloatPacked;
+	StreamDescription.mChannelsPerFrame = NumChannels;
+	StreamDescription.mBitsPerChannel   = 8 * BytesPerSample;
+	StreamDescription.mBytesPerFrame    = BytesPerSample * StreamDescription.mChannelsPerFrame;
+	StreamDescription.mFramesPerPacket  = 1;
+	StreamDescription.mBytesPerPacket   = StreamDescription.mFramesPerPacket * StreamDescription.mBytesPerFrame;
+	
+	// Configure output format
+	Status = AudioUnitSetProperty(IOUnit,
 		kAudioUnitProperty_StreamFormat,
 		kAudioUnitScope_Output,
-		0,
+		kInputBus,
 		&StreamDescription,
 		sizeof(StreamDescription));
-	check(AudioUnitStatus == noErr);
+	check(Status == noErr);
 
+	// Setup capture callback
 	AURenderCallbackStruct CallbackInfo;
 	CallbackInfo.inputProc = RecordingCallback;
 	CallbackInfo.inputProcRefCon = this;
-	AudioUnitStatus = AudioUnitSetProperty(AudioUnit,
+	Status = AudioUnitSetProperty(IOUnit,
 		kAudioOutputUnitProperty_SetInputCallback,
 		kAudioUnitScope_Global,
-		0,
+		kInputBus,
 		&CallbackInfo,
 		sizeof(CallbackInfo));
-	check(AudioUnitStatus == noErr);
+	check(Status == noErr);
+	
+	// Initialize audio unit
+	Status = AudioUnitInitialize(IOUnit);
+	check(Status == noErr);
 
-	AudioUnitStatus = AudioUnitInitialize(AudioUnit);
-	check(AudioUnitStatus == noErr);
+	// Configure unit processing
+	SetHardwareFeatureEnabled(Audio::EHardwareInputFeature::EchoCancellation, InParams.bUseHardwareAEC);
+	SetHardwareFeatureEnabled(Audio::EHardwareInputFeature::AutomaticGainControl, InParams.bUseHardwareAEC);
 
-	return AudioUnitStatus == noErr;
+	return Status == noErr;
 }
 
 bool Audio::FAudioCaptureAudioUnitStream::CloseStream()
 {
 	StopStream();
-	AudioComponentInstanceDispose(AudioUnit);
+	AudioComponentInstanceDispose(IOUnit);
 	return true;
 }
 
 bool Audio::FAudioCaptureAudioUnitStream::StartStream()
 {
-	AudioUnitStatus = AudioOutputUnitStart(AudioUnit);
-	return AudioUnitStatus == noErr;
+	return (AudioOutputUnitStart(IOUnit) == noErr);
 }
 
 bool Audio::FAudioCaptureAudioUnitStream::StopStream()
 {
-	AudioUnitStatus = AudioOutputUnitStop(AudioUnit);
-	return AudioUnitStatus == noErr;
+	return (AudioOutputUnitStop(IOUnit) == noErr);
 }
 
 bool Audio::FAudioCaptureAudioUnitStream::AbortStream()
@@ -164,7 +212,7 @@ bool Audio::FAudioCaptureAudioUnitStream::IsCapturing() const
 void Audio::FAudioCaptureAudioUnitStream::OnAudioCapture(void* InBuffer, uint32 InBufferFrames, double StreamTime, bool bOverflow)
 {
 	float* InBufferData = (float*)InBuffer;
-	OnCapture(InBufferData, InBufferFrames, NumChannels, StreamTime, bOverflow);
+	OnCapture(InBufferData, InBufferFrames, NumChannels, SampleRate, StreamTime, bOverflow);
 }
 
 bool Audio::FAudioCaptureAudioUnitStream::GetInputDevicesAvailable(TArray<FCaptureDeviceInfo>& OutDevices)
@@ -177,3 +225,42 @@ bool Audio::FAudioCaptureAudioUnitStream::GetInputDevicesAvailable(TArray<FCaptu
 
 	return true;
 }
+
+void Audio::FAudioCaptureAudioUnitStream::SetHardwareFeatureEnabled(EHardwareInputFeature FeatureType, bool bEnabled)
+{
+	if (IOUnit == nil)
+	{
+		return;
+	}
+
+	// we ignore result those functions because sometime we can't set this parameters
+	OSStatus status = noErr;
+
+	switch(FeatureType)
+	{
+		case Audio::EHardwareInputFeature::EchoCancellation:
+		{
+			const UInt32 EnableParam = (bEnabled ? 0 : 1);
+			status = AudioUnitSetProperty(IOUnit,
+			  kAUVoiceIOProperty_BypassVoiceProcessing,
+			  kAudioUnitScope_Global,
+			  kInputBus,
+			  &EnableParam,
+			  sizeof(EnableParam)
+			);
+		}
+			break;
+		case Audio::EHardwareInputFeature::AutomaticGainControl:
+		{
+			const UInt32 EnableParam = (bEnabled ? 1 : 0);
+			status = AudioUnitSetProperty(IOUnit,
+			  kAUVoiceIOProperty_VoiceProcessingEnableAGC,
+			  kAudioUnitScope_Global,
+			  kInputBus,
+			  &EnableParam,
+			  sizeof(EnableParam)
+			);
+		}
+			break;
+	}
+}
diff --git a/Engine/Source/Runtime/AudioCaptureImplementations/IOS/AudioCaptureAudioUnit/Private/AudioCaptureAudioUnit.h b/Engine/Source/Runtime/AudioCaptureImplementations/IOS/AudioCaptureAudioUnit/Private/AudioCaptureAudioUnit.h
--- a/Engine/Source/Runtime/AudioCaptureImplementations/IOS/AudioCaptureAudioUnit/Private/AudioCaptureAudioUnit.h
+++ b/Engine/Source/Runtime/AudioCaptureImplementations/IOS/AudioCaptureAudioUnit/Private/AudioCaptureAudioUnit.h
@@ -24,14 +24,19 @@ namespace Audio
 		virtual bool IsCapturing() const override;
 		virtual void OnAudioCapture(void* InBuffer, uint32 InBufferFrames, double StreamTime, bool bOverflow) override;
 		virtual bool GetInputDevicesAvailable(TArray<FCaptureDeviceInfo>& OutDevices) override;
+		virtual void SetHardwareFeatureEnabled(EHardwareInputFeature FeatureType, bool bEnabled);
 
-		AudioComponentInstance AudioUnit;
+		OSStatus OnCaptureCallback( AudioUnitRenderActionFlags* ioActionFlags, const AudioTimeStamp* inTimeStamp, UInt32 inBusNumber, UInt32 inNumberFrames, AudioBufferList* ioData );
 
 	private:
+		void AllocateBuffer(int SizeInBytes);
+		
+		AudioComponentInstance IOUnit;
+		
 		FOnCaptureFunction OnCapture;
 		int32 NumChannels;
 		int32 SampleRate;
-
-		OSStatus AudioUnitStatus;
+		TArray<uint8> CaptureBuffer;
+		int BufferSize = 0;
 	};
-}
\ No newline at end of file
+}
