Fix for crash in Audio Capture Component when OpenCaptureStream() hasn't been called

diff --git a/Engine/Source/Runtime/AudioCaptureImplementations/IOS/AudioCaptureAudioUnit/Private/AudioCaptureAudioUnit.cpp b/Engine/Source/Runtime/AudioCaptureImplementations/IOS/AudioCaptureAudioUnit/Private/AudioCaptureAudioUnit.cpp
--- a/Engine/Source/Runtime/AudioCaptureImplementations/IOS/AudioCaptureAudioUnit/Private/AudioCaptureAudioUnit.cpp
+++ b/Engine/Source/Runtime/AudioCaptureImplementations/IOS/AudioCaptureAudioUnit/Private/AudioCaptureAudioUnit.cpp
@@ -8,6 +8,8 @@ const int32 kOutputBus = 0;
 Audio::FAudioCaptureAudioUnitStream::FAudioCaptureAudioUnitStream()
 	: NumChannels(0)
 	, SampleRate(0)
+	, bIsStreamOpen(false)
+	, bWasStreamStarted(false)
 {
 }
 
@@ -166,23 +168,29 @@ bool Audio::FAudioCaptureAudioUnitStream::OpenCaptureStream(const FAudioCaptureD
 	SetHardwareFeatureEnabled(Audio::EHardwareInputFeature::EchoCancellation, InParams.bUseHardwareAEC);
 	SetHardwareFeatureEnabled(Audio::EHardwareInputFeature::AutomaticGainControl, InParams.bUseHardwareAEC);
 
-	return Status == noErr;
+	bIsStreamOpen = (Status == noErr);
+
+	return bIsStreamOpen;
 }
 
 bool Audio::FAudioCaptureAudioUnitStream::CloseStream()
 {
 	StopStream();
 	AudioComponentInstanceDispose(IOUnit);
+	bIsStreamOpen = false;
+
 	return true;
 }
 
 bool Audio::FAudioCaptureAudioUnitStream::StartStream()
 {
-	return (AudioOutputUnitStart(IOUnit) == noErr);
+	bHasCaptureStarted = (AudioOutputUnitStart(IOUnit) == noErr);
+	return bHasCaptureStarted;
 }
 
 bool Audio::FAudioCaptureAudioUnitStream::StopStream()
 {
+	bHasCaptureStarted = false;
 	return (AudioOutputUnitStop(IOUnit) == noErr);
 }
 
@@ -201,12 +209,12 @@ bool Audio::FAudioCaptureAudioUnitStream::GetStreamTime(double& OutStreamTime)
 
 bool Audio::FAudioCaptureAudioUnitStream::IsStreamOpen() const
 {
-	return true;
+	return bIsStreamOpen;
 }
 
 bool Audio::FAudioCaptureAudioUnitStream::IsCapturing() const
 {
-	return true;
+	return bHasCaptureStarted;
 }
 
 void Audio::FAudioCaptureAudioUnitStream::OnAudioCapture(void* InBuffer, uint32 InBufferFrames, double StreamTime, bool bOverflow)
diff --git a/Engine/Source/Runtime/AudioCaptureImplementations/IOS/AudioCaptureAudioUnit/Private/AudioCaptureAudioUnit.h b/Engine/Source/Runtime/AudioCaptureImplementations/IOS/AudioCaptureAudioUnit/Private/AudioCaptureAudioUnit.h
--- a/Engine/Source/Runtime/AudioCaptureImplementations/IOS/AudioCaptureAudioUnit/Private/AudioCaptureAudioUnit.h
+++ b/Engine/Source/Runtime/AudioCaptureImplementations/IOS/AudioCaptureAudioUnit/Private/AudioCaptureAudioUnit.h
@@ -38,5 +38,7 @@ namespace Audio
 		int32 SampleRate;
 		TArray<uint8> CaptureBuffer;
 		int BufferSize = 0;
+		bool bIsStreamOpen;
+		bool bHasCaptureStarted;
 	};
 }
