Fix RtAudio in modular builds on some platform extensions.

#rnx

diff --git a/Engine/Source/Runtime/AudioCaptureImplementations/AudioCaptureRtAudio/AudioCaptureRtAudio.Build.cs b/Engine/Source/Runtime/AudioCaptureImplementations/AudioCaptureRtAudio/AudioCaptureRtAudio.Build.cs
--- a/Engine/Source/Runtime/AudioCaptureImplementations/AudioCaptureRtAudio/AudioCaptureRtAudio.Build.cs
+++ b/Engine/Source/Runtime/AudioCaptureImplementations/AudioCaptureRtAudio/AudioCaptureRtAudio.Build.cs
@@ -4,6 +4,8 @@ using UnrealBuildTool;
 
 public class AudioCaptureRtAudio : ModuleRules
 {
+	protected virtual bool WithRtAudio { get => false; }
+
 	public AudioCaptureRtAudio(ReadOnlyTargetRules Target) : base(Target)
 	{
 		PrivateDependencyModuleNames.Add("Core");
@@ -25,7 +27,7 @@ public class AudioCaptureRtAudio : ModuleRules
         }
 		else
 		{
-			PublicDefinitions.Add("WITH_RTAUDIO=0");
+			PublicDefinitions.Add($"WITH_RTAUDIO={(WithRtAudio?"1":"0")}");
 		}
 	}
 }
diff --git a/Engine/Source/Runtime/AudioCaptureImplementations/AudioCaptureRtAudio/Private/RtAudio.cpp b/Engine/Source/Runtime/AudioCaptureImplementations/AudioCaptureRtAudio/Private/RtAudio.cpp
--- a/Engine/Source/Runtime/AudioCaptureImplementations/AudioCaptureRtAudio/Private/RtAudio.cpp
+++ b/Engine/Source/Runtime/AudioCaptureImplementations/AudioCaptureRtAudio/Private/RtAudio.cpp
@@ -3731,12 +3731,14 @@ static const char* getAsioErrorString( ASIOError result )
 
 //=============================================================================
 
-#define SAFE_RELEASE( objectPtr )\
-if ( objectPtr )\
-{\
-  objectPtr->Release();\
-  objectPtr = NULL;\
-}
+#if !defined(SAFE_RELEASE)
+	#define SAFE_RELEASE( objectPtr )\
+	if ( objectPtr )\
+	{\
+	  objectPtr->Release();\
+	  objectPtr = NULL;\
+	}
+#endif
 
 typedef HANDLE ( __stdcall *TAvSetMmThreadCharacteristicsPtr )( LPCWSTR TaskName, LPDWORD TaskIndex );
 
