- Replace CoreAudio API constant with version not containing bad substrings

diff --git a/Engine/Source/Runtime/AudioCaptureImplementations/AudioCaptureRtAudio/Private/RtAudio.cpp b/Engine/Source/Runtime/AudioCaptureImplementations/AudioCaptureRtAudio/Private/RtAudio.cpp
--- a/Engine/Source/Runtime/AudioCaptureImplementations/AudioCaptureRtAudio/Private/RtAudio.cpp
+++ b/Engine/Source/Runtime/AudioCaptureImplementations/AudioCaptureRtAudio/Private/RtAudio.cpp
@@ -538,7 +538,7 @@ RtApiCore:: RtApiCore()
   CFRunLoopRef theRunLoop = NULL;
   AudioObjectPropertyAddress property = { kAudioHardwarePropertyRunLoop,
                                           kAudioObjectPropertyScopeGlobal,
-                                          kAudioObjectPropertyElementMaster };
+										  kAudioObjectPropertyElementMain };
   OSStatus result = AudioObjectSetPropertyData( kAudioObjectSystemObject, &property, 0, NULL, sizeof(CFRunLoopRef), &theRunLoop);
   if ( result != noErr ) {
     errorText_ = "RtApiCore::RtApiCore: error setting run loop property!";
@@ -559,7 +559,7 @@ unsigned int RtApiCore :: getDeviceCount( void )
 {
   // Find out how many audio devices there are, if any.
   UInt32 dataSize;
-  AudioObjectPropertyAddress propertyAddress = { kAudioHardwarePropertyDevices, kAudioObjectPropertyScopeGlobal, kAudioObjectPropertyElementMaster };
+  AudioObjectPropertyAddress propertyAddress = { kAudioHardwarePropertyDevices, kAudioObjectPropertyScopeGlobal, kAudioObjectPropertyElementMain };
   OSStatus result = AudioObjectGetPropertyDataSize( kAudioObjectSystemObject, &propertyAddress, 0, NULL, &dataSize );
   if ( result != noErr ) {
     errorText_ = "RtApiCore::getDeviceCount: OS-X error getting device info!";
@@ -577,7 +577,7 @@ unsigned int RtApiCore :: getDefaultInputDevice( void )
 
   AudioDeviceID id;
   UInt32 dataSize = sizeof( AudioDeviceID );
-  AudioObjectPropertyAddress property = { kAudioHardwarePropertyDefaultInputDevice, kAudioObjectPropertyScopeGlobal, kAudioObjectPropertyElementMaster };
+  AudioObjectPropertyAddress property = { kAudioHardwarePropertyDefaultInputDevice, kAudioObjectPropertyScopeGlobal, kAudioObjectPropertyElementMain };
   OSStatus result = AudioObjectGetPropertyData( kAudioObjectSystemObject, &property, 0, NULL, &dataSize, &id );
   if ( result != noErr ) {
     errorText_ = "RtApiCore::getDefaultInputDevice: OS-X system error getting device.";
@@ -610,7 +610,7 @@ unsigned int RtApiCore :: getDefaultOutputDevice( void )
 
   AudioDeviceID id;
   UInt32 dataSize = sizeof( AudioDeviceID );
-  AudioObjectPropertyAddress property = { kAudioHardwarePropertyDefaultOutputDevice, kAudioObjectPropertyScopeGlobal, kAudioObjectPropertyElementMaster };
+  AudioObjectPropertyAddress property = { kAudioHardwarePropertyDefaultOutputDevice, kAudioObjectPropertyScopeGlobal, kAudioObjectPropertyElementMain };
   OSStatus result = AudioObjectGetPropertyData( kAudioObjectSystemObject, &property, 0, NULL, &dataSize, &id );
   if ( result != noErr ) {
     errorText_ = "RtApiCore::getDefaultOutputDevice: OS-X system error getting device.";
@@ -659,7 +659,7 @@ RtAudio::DeviceInfo RtApiCore :: getDeviceInfo( unsigned int device )
   UInt32 dataSize = sizeof( AudioDeviceID ) * nDevices;
   AudioObjectPropertyAddress property = { kAudioHardwarePropertyDevices,
                                           kAudioObjectPropertyScopeGlobal,
-                                          kAudioObjectPropertyElementMaster };
+										  kAudioObjectPropertyElementMain };
   OSStatus result = AudioObjectGetPropertyData( kAudioObjectSystemObject, &property,
                                                 0, NULL, &dataSize, (void *) &deviceList );
   if ( result != noErr ) {
@@ -908,7 +908,7 @@ static OSStatus rateListener( AudioObjectID inDevice,
   UInt32 dataSize = sizeof( Float64 );
   AudioObjectPropertyAddress property = { kAudioDevicePropertyNominalSampleRate,
                                           kAudioObjectPropertyScopeGlobal,
-                                          kAudioObjectPropertyElementMaster };
+										  kAudioObjectPropertyElementMain };
   AudioObjectGetPropertyData( inDevice, &property, 0, NULL, &dataSize, rate );
   return kAudioHardwareNoError;
 }
@@ -936,7 +936,7 @@ bool RtApiCore :: probeDeviceOpen( unsigned int device, StreamMode mode, unsigne
   UInt32 dataSize = sizeof( AudioDeviceID ) * nDevices;
   AudioObjectPropertyAddress property = { kAudioHardwarePropertyDevices,
                                           kAudioObjectPropertyScopeGlobal,
-                                          kAudioObjectPropertyElementMaster };
+										  kAudioObjectPropertyElementMain };
   OSStatus result = AudioObjectGetPropertyData( kAudioObjectSystemObject, &property,
                                                 0, NULL, &dataSize, (void *) &deviceList );
   if ( result != noErr ) {
@@ -1127,7 +1127,7 @@ bool RtApiCore :: probeDeviceOpen( unsigned int device, StreamMode mode, unsigne
 
     // Set a property listener for the sample rate change
     Float64 reportedRate = 0.0;
-    AudioObjectPropertyAddress tmp = { kAudioDevicePropertyNominalSampleRate, kAudioObjectPropertyScopeGlobal, kAudioObjectPropertyElementMaster };
+    AudioObjectPropertyAddress tmp = { kAudioDevicePropertyNominalSampleRate, kAudioObjectPropertyScopeGlobal, kAudioObjectPropertyElementMain };
     result = AudioObjectAddPropertyListener( id, &tmp, rateListener, (void *) &reportedRate );
     if ( result != noErr ) {
       errorStream_ << "RtApiCore::probeDeviceOpen: system error (" << getErrorCode( result ) << ") setting sample rate property listener for device (" << device << ").";
@@ -1445,7 +1445,7 @@ void RtApiCore :: closeStream( void )
     if (handle) {
       AudioObjectPropertyAddress property = { kAudioHardwarePropertyDevices,
         kAudioObjectPropertyScopeGlobal,
-        kAudioObjectPropertyElementMaster };
+		kAudioObjectPropertyElementMain };
 
       property.mSelector = kAudioDeviceProcessorOverload;
       property.mScope = kAudioObjectPropertyScopeGlobal;
@@ -1468,7 +1468,7 @@ void RtApiCore :: closeStream( void )
     if (handle) {
       AudioObjectPropertyAddress property = { kAudioHardwarePropertyDevices,
         kAudioObjectPropertyScopeGlobal,
-        kAudioObjectPropertyElementMaster };
+		kAudioObjectPropertyElementMain };
 
       property.mSelector = kAudioDeviceProcessorOverload;
       property.mScope = kAudioObjectPropertyScopeGlobal;
diff --git a/Engine/Source/Runtime/Online/Voice/Private/Mac/VoiceModuleMac.cpp b/Engine/Source/Runtime/Online/Voice/Private/Mac/VoiceModuleMac.cpp
--- a/Engine/Source/Runtime/Online/Voice/Private/Mac/VoiceModuleMac.cpp
+++ b/Engine/Source/Runtime/Online/Voice/Private/Mac/VoiceModuleMac.cpp
@@ -112,7 +112,7 @@ public:
 		
 		uint32 PropSize = sizeof(AudioDeviceID);
 		AudioDeviceID InputDevice = 0;
-		AudioObjectPropertyAddress Address = { kAudioHardwarePropertyDefaultInputDevice, kAudioObjectPropertyScopeGlobal, kAudioObjectPropertyElementMaster };
+		AudioObjectPropertyAddress Address = { kAudioHardwarePropertyDefaultInputDevice, kAudioObjectPropertyScopeGlobal, kAudioObjectPropertyElementMain };
 		if ( AudioObjectGetPropertyData(kAudioObjectSystemObject, &Address, 0, nullptr, &PropSize, &InputDevice) != 0 || InputDevice == kAudioDeviceUnknown )
 		{
 			UE_LOG(LogVoiceCapture, Warning, TEXT("Couldn't get Default CoreAudio input device"));
