Fix CIS build error on Mac Debug builds of shootergame between RtAudio and ResonanceAudio plugins
- Removed some unneeded iostream usage in RtAudio
- Removed exposure of private headers from RtAudio
- Removed exposure of private headers from ResonanceAudio
#rnx

diff --git a/Engine/Plugins/Runtime/AudioCapture/Source/AudioCaptureEditor/AudioCaptureEditor.build.cs b/Engine/Plugins/Runtime/AudioCapture/Source/AudioCaptureEditor/AudioCaptureEditor.build.cs
--- a/Engine/Plugins/Runtime/AudioCapture/Source/AudioCaptureEditor/AudioCaptureEditor.build.cs
+++ b/Engine/Plugins/Runtime/AudioCapture/Source/AudioCaptureEditor/AudioCaptureEditor.build.cs
@@ -21,14 +21,6 @@ public class AudioCaptureEditor : ModuleRules
 			}
 		);
 
-		PrivateIncludePaths.AddRange(
-			new string[] {
-				"AudioCaptureEditor/Private",
-				"AudioCapture/Private",
-				"../../../../Source/Runtime/AudioCaptureImplementations/AudioCaptureRtAudio/Private" // This is required to include RtAudio.h in AudioRecordingManager.h.
-			}
-		);
-
 		if (Target.Platform == UnrealTargetPlatform.Win64)
 		{
 			// Allow us to use direct sound
diff --git a/Engine/Plugins/Runtime/AudioCapture/Source/AudioCaptureEditor/Private/AudioRecordingManager.cpp b/Engine/Plugins/Runtime/AudioCapture/Source/AudioCaptureEditor/Private/AudioRecordingManager.cpp
--- a/Engine/Plugins/Runtime/AudioCapture/Source/AudioCaptureEditor/Private/AudioRecordingManager.cpp
+++ b/Engine/Plugins/Runtime/AudioCapture/Source/AudioCaptureEditor/Private/AudioRecordingManager.cpp
@@ -21,16 +21,16 @@ namespace Audio
 /**
 * Callback Function For the Microphone Capture for RtAudio
 */
-static int32 OnAudioCaptureCallback(void *OutBuffer, void* InBuffer, uint32 InBufferFrames, double StreamTime, RtAudioStreamStatus AudioStreamStatus, void* InUserData)
+static int32 OnAudioCaptureCallback(void *OutBuffer, void* InBuffer, uint32 InBufferFrames, double StreamTime, uint32 AudioStreamStatus, void* InUserData)
 {
 	// Check for stream overflow (i.e. we're feeding audio faster than we're consuming)
-
+	const bool bIsOverflow = 0 != AudioStreamStatus;
 
 	// Cast the user data to the singleton mic recorder
 	FAudioRecordingManager* AudioRecordingManager = (FAudioRecordingManager*)InUserData;
 
 	// Call the mic capture callback function
-	return AudioRecordingManager->OnAudioCapture(InBuffer, InBufferFrames, StreamTime, AudioStreamStatus == RTAUDIO_INPUT_OVERFLOW);
+	return AudioRecordingManager->OnAudioCapture(InBuffer, InBufferFrames, StreamTime, bIsOverflow);
 }
 
 /**
@@ -51,9 +51,9 @@ FAudioRecordingManager::FAudioRecordingManager()
 
 FAudioRecordingManager::~FAudioRecordingManager()
 {
-	if (ADC.isStreamOpen())
+	if (ADC.IsStreamOpen())
 	{
-		ADC.abortStream();
+		ADC.AbortStream();
 	}
 }
 
@@ -78,20 +78,10 @@ void FAudioRecordingManager::StartRecording(const FRecordingSettings& InSettings
 	RecordingBlockSize = 1024;
 
 	// If we have a stream open close it (reusing streams can cause a blip of previous recordings audio)
-	if (ADC.isStreamOpen())
+	if (ADC.IsStreamOpen())
 	{
-		try
-		{
-			ADC.stopStream();
-			ADC.closeStream();
-		}
-		catch (RtAudioError& e)
-		{
-			FString ErrorMessage = FString(e.what());
-			bError = true;
-			UE_LOG(LogMicManager, Error, TEXT("Failed to close the mic capture device stream: %s"), *ErrorMessage);
-			return;
-		}
+		ADC.StopStream();
+		ADC.CloseStream();
 	}
 
 	UE_LOG(LogMicManager, Log, TEXT("Starting mic recording."));
@@ -103,9 +93,9 @@ void FAudioRecordingManager::StartRecording(const FRecordingSettings& InSettings
 	InputGain = Audio::ConvertToLinear(Settings.GainDb);
 
 	// Get the default mic input device info
-	RtAudio::DeviceInfo Info = ADC.getDeviceInfo(StreamParams.deviceId);
-	RecordingSampleRate = Info.preferredSampleRate;
-	NumInputChannels = Info.inputChannels;	
+	FRtAudioInputWrapper::FDeviceInfo Info = ADC.GetDeviceInfo(StreamParams.DeviceID);
+	RecordingSampleRate = Info.PreferredSampleRate;
+	NumInputChannels = Info.NumChannels;	
 
 	if (NumInputChannels == 0)
 	{
@@ -144,40 +134,23 @@ void FAudioRecordingManager::StartRecording(const FRecordingSettings& InSettings
 	// Publish to the mic input thread that we're ready to record...
 	bRecording = true;
 
-	StreamParams.deviceId = ADC.getDefaultInputDevice(); // Only use the default input device for now
+	StreamParams.DeviceID = ADC.GetDefaultInputDevice(); // Only use the default input device for now
 
-	StreamParams.nChannels = NumInputChannels;
-	StreamParams.firstChannel = 0;
+	StreamParams.NumChannels = NumInputChannels;
 
 	uint32 BufferFrames = FMath::Max(RecordingBlockSize, 256);
 
-	UE_LOG(LogMicManager, Log, TEXT("Initialized mic recording manager at %d hz sample rate, %d channels, and %d Recording Block Size"), RecordingSampleRate, StreamParams.nChannels, BufferFrames);
+	UE_LOG(LogMicManager, Log, TEXT("Initialized mic recording manager at %d hz sample rate, %d channels, and %d Recording Block Size"), RecordingSampleRate, StreamParams.NumChannels, BufferFrames);
 
 	// RtAudio uses exceptions for error handling... 
-	try
-	{
-		// Open up new audio stream
-		ADC.openStream(nullptr, &StreamParams, RTAUDIO_SINT16, RecordingSampleRate, &BufferFrames, &OnAudioCaptureCallback, this);
-	}
-	catch (RtAudioError& e)
-	{
-		FString ErrorMessage = FString(e.what());
-		bError = true;
-		UE_LOG(LogMicManager, Error, TEXT("Failed to open the mic capture device: %s"), *ErrorMessage);
-	}
-	catch (const std::exception& e)
-	{
-		bError = true;
-		FString ErrorMessage = FString(e.what());
-		UE_LOG(LogMicManager, Error, TEXT("Failed to open the mic capture device: %s"), *ErrorMessage);
-	}
-	catch (...)
+	// Open up new audio stream
+	bool bSuccess = ADC.OpenStream(StreamParams, RecordingSampleRate, &BufferFrames, &OnAudioCaptureCallback, this);
+	bError = !bSuccess;
+
+	if (bSuccess)
 	{
-		UE_LOG(LogMicManager, Error, TEXT("Failed to open the mic capture device: unknown error"));
-		bError = true;
+		ADC.StartStream();
 	}
-
-	ADC.startStream();
 }
 
 static void SampleRateConvert(float CurrentSR, float TargetSR, int32 NumChannels, const TArray<int16>& CurrentRecordedPCMData, int32 NumSamplesToConvert, TArray<int16>& OutConverted)
@@ -235,8 +208,8 @@ void FAudioRecordingManager::StopRecording(TArray<USoundWave*>& OutSoundWaves)
 	if (bRecording)
 	{
 		bRecording = false;
-		ADC.stopStream();
-		ADC.closeStream();
+		ADC.StopStream();
+		ADC.CloseStream();
 
 		if (CurrentRecordedPCMData.Num() > 0)
 		{
@@ -464,4 +437,4 @@ int32 FAudioRecordingManager::OnAudioCapture(void* InBuffer, uint32 InBufferFram
 	return 1;
 }
 
-} // namespace Audio
\ No newline at end of file
+} // namespace Audio
diff --git a/Engine/Plugins/Runtime/AudioCapture/Source/AudioCaptureEditor/Private/AudioRecordingManager.h b/Engine/Plugins/Runtime/AudioCapture/Source/AudioCaptureEditor/Private/AudioRecordingManager.h
--- a/Engine/Plugins/Runtime/AudioCapture/Source/AudioCaptureEditor/Private/AudioRecordingManager.h
+++ b/Engine/Plugins/Runtime/AudioCapture/Source/AudioCaptureEditor/Private/AudioRecordingManager.h
@@ -5,12 +5,7 @@
 #include "CoreMinimal.h"
 #include "HAL/ThreadSafeBool.h"
 #include "Engine/EngineTypes.h"
-#if PLATFORM_WINDOWS
-	#include "Windows/WindowsHWrapper.h"
-#endif
-THIRD_PARTY_INCLUDES_START
-	#include "RtAudio.h"
-THIRD_PARTY_INCLUDES_END
+#include "RtAudioWrapper.h"
 
 class USoundWave;
 
@@ -78,10 +73,10 @@ namespace Audio
 		FRecordingSettings Settings;
 
 		// RtAudio ADC object -- used to interact with low-level audio device.
-		RtAudio ADC;
+		FRtAudioInputWrapper ADC;
 
 		// Stream parameters to initialize the ADC
-		RtAudio::StreamParameters StreamParams;
+		FRtAudioInputWrapper::FStreamParameters StreamParams;
 
 		// Critical section used to stop and retrieve finished audio buffers.
 		FCriticalSection CriticalSection;
diff --git a/Engine/Plugins/Runtime/ResonanceAudio/Source/ResonanceAudio/Private/ResonanceAudioAmbisonics.cpp b/Engine/Plugins/Runtime/ResonanceAudio/Source/ResonanceAudio/Private/ResonanceAudioAmbisonics.cpp
--- a/Engine/Plugins/Runtime/ResonanceAudio/Source/ResonanceAudio/Private/ResonanceAudioAmbisonics.cpp
+++ b/Engine/Plugins/Runtime/ResonanceAudio/Source/ResonanceAudio/Private/ResonanceAudioAmbisonics.cpp
@@ -3,6 +3,7 @@
 //
 #include "ResonanceAudioAmbisonics.h"
 #include "ResonanceAudioAmbisonicsSettings.h"
+#include "ResonanceAudioAmbisonicsSettingsProxy.h"
 #include "ResonanceAudioPluginListener.h"
 #include "SoundFieldRendering.h"
 #include "AudioPluginUtilities.h"
diff --git a/Engine/Plugins/Runtime/ResonanceAudio/Source/ResonanceAudio/Private/ResonanceAudioAmbisonicsSettings.cpp b/Engine/Plugins/Runtime/ResonanceAudio/Source/ResonanceAudio/Private/ResonanceAudioAmbisonicsSettings.cpp
new file mode 100644
--- /dev/null
+++ b/Engine/Plugins/Runtime/ResonanceAudio/Source/ResonanceAudio/Private/ResonanceAudioAmbisonicsSettings.cpp
@@ -0,0 +1,51 @@
+//
+// Copyright (C) Google Inc. 2017. All rights reserved.
+//
+
+#include "ResonanceAudioAmbisonicsSettings.h"
+
+#include "AudioDevice.h"
+#include "HAL/LowLevelMemTracker.h"
+#include "ResonanceAudioAmbisonicsSettingsProxy.h"
+#include "ResonanceAudioCommon.h"
+
+TUniquePtr<ISoundfieldEncodingSettingsProxy> UResonanceAudioSoundfieldSettings::GetProxy() const
+{
+	LLM_SCOPE_BYTAG(AudioSpatializationPlugins);
+	FResonanceAmbisonicsSettingsProxy* Proxy = new FResonanceAmbisonicsSettingsProxy();
+
+	switch (RenderMode)
+	{
+		case EResonanceRenderMode::StereoPanning:
+		{
+			Proxy->RenderingMode = vraudio::kStereoPanning;
+			break;
+		}
+		case EResonanceRenderMode::BinauralLowQuality:
+		{
+			Proxy->RenderingMode = vraudio::kBinauralLowQuality;
+			break;
+		}
+		case EResonanceRenderMode::BinauralMediumQuality:
+		{
+			Proxy->RenderingMode = vraudio::kBinauralMediumQuality;
+			break;
+		}
+		case EResonanceRenderMode::BinauralHighQuality:
+		{
+			Proxy->RenderingMode = vraudio::kBinauralHighQuality;
+			break;
+		}
+		case EResonanceRenderMode::RoomEffectsOnly:
+		{
+			Proxy->RenderingMode = vraudio::kRoomEffectsOnly;
+			break;
+		}
+		default:
+		{
+			checkNoEntry();
+		}
+	}
+
+	return TUniquePtr<ISoundfieldEncodingSettingsProxy>(Proxy);
+}
diff --git a/Engine/Plugins/Runtime/ResonanceAudio/Source/ResonanceAudio/Private/ResonanceAudioAmbisonicsSettingsProxy.h b/Engine/Plugins/Runtime/ResonanceAudio/Source/ResonanceAudio/Private/ResonanceAudioAmbisonicsSettingsProxy.h
new file mode 100644
--- /dev/null
+++ b/Engine/Plugins/Runtime/ResonanceAudio/Source/ResonanceAudio/Private/ResonanceAudioAmbisonicsSettingsProxy.h
@@ -0,0 +1,30 @@
+//
+// Copyright (C) Google Inc. 2017. All rights reserved.
+//
+
+#pragma once
+
+#include "ResonanceAudioAmbisonicsSettings.h"
+
+#include "AudioDevice.h"
+#include "HAL/LowLevelMemTracker.h"
+#include "ResonanceAudioCommon.h"
+
+class FResonanceAmbisonicsSettingsProxy : public ISoundfieldEncodingSettingsProxy
+{
+public:
+	vraudio::RenderingMode RenderingMode;
+
+	virtual uint32 GetUniqueId() const override
+	{
+		return static_cast<uint32>(RenderingMode);
+	}
+
+	virtual TUniquePtr<ISoundfieldEncodingSettingsProxy> Duplicate() const override
+	{
+		LLM_SCOPE_BYTAG(AudioSpatializationPlugins);
+		FResonanceAmbisonicsSettingsProxy* Proxy = new FResonanceAmbisonicsSettingsProxy();
+		Proxy->RenderingMode = RenderingMode;
+		return TUniquePtr<ISoundfieldEncodingSettingsProxy>(Proxy);
+	}
+};
diff --git a/Engine/Plugins/Runtime/ResonanceAudio/Source/ResonanceAudio/Private/ResonanceAudioCommon.cpp b/Engine/Plugins/Runtime/ResonanceAudio/Source/ResonanceAudio/Private/ResonanceAudioCommon.cpp
--- a/Engine/Plugins/Runtime/ResonanceAudio/Source/ResonanceAudio/Private/ResonanceAudioCommon.cpp
+++ b/Engine/Plugins/Runtime/ResonanceAudio/Source/ResonanceAudio/Private/ResonanceAudioCommon.cpp
@@ -3,6 +3,7 @@
 //
 
 #include "ResonanceAudioCommon.h"
+#include "ResonanceAudioConstants.h"
 #include "Misc/Paths.h"
 #include "EngineUtils.h"
 #include "HAL/FileManager.h"
diff --git a/Engine/Plugins/Runtime/ResonanceAudio/Source/ResonanceAudio/Private/ResonanceAudioCommon.h b/Engine/Plugins/Runtime/ResonanceAudio/Source/ResonanceAudio/Private/ResonanceAudioCommon.h
--- a/Engine/Plugins/Runtime/ResonanceAudio/Source/ResonanceAudio/Private/ResonanceAudioCommon.h
+++ b/Engine/Plugins/Runtime/ResonanceAudio/Source/ResonanceAudio/Private/ResonanceAudioCommon.h
@@ -14,6 +14,7 @@ THIRD_PARTY_INCLUDES_END
 
 
 #include "UObject/ObjectMacros.h"
+#include "ResonanceAudioConstants.h"
 #include "ResonanceAudioEnums.h"
 #include "ISoundfieldFormat.h"
 
diff --git a/Engine/Plugins/Runtime/ResonanceAudio/Source/ResonanceAudio/Private/ResonanceAudioConstants.h b/Engine/Plugins/Runtime/ResonanceAudio/Source/ResonanceAudio/Private/ResonanceAudioConstants.h
new file mode 100644
--- /dev/null
+++ b/Engine/Plugins/Runtime/ResonanceAudio/Source/ResonanceAudio/Private/ResonanceAudioConstants.h
@@ -0,0 +1,19 @@
+//
+// Copyright Google Inc. 2017. All rights reserved.
+//
+
+#pragma once
+
+#include "Math/Color.h"
+
+namespace ResonanceAudio
+{
+	// Resonance Audio assets base color.
+	const FColor ASSET_COLOR = FColor(0, 198, 246);
+
+	// Number of surfaces in a shoe-box room.
+	constexpr int NUM_SURFACES = 6;
+
+	// Conversion factor between Resonance Audio and Unreal world distance units (1cm in Unreal = 0.01m in Resonance Audio).
+	constexpr float SCALE_FACTOR = 0.01f;
+}
diff --git a/Engine/Plugins/Runtime/ResonanceAudio/Source/ResonanceAudio/Private/ResonanceAudioEnums.h b/Engine/Plugins/Runtime/ResonanceAudio/Source/ResonanceAudio/Private/ResonanceAudioEnums.h
deleted file mode 100644
diff --git a/Engine/Plugins/Runtime/ResonanceAudio/Source/ResonanceAudio/Private/ResonanceAudioLibrary/resonance_audio/utils/wav_reader.h b/Engine/Plugins/Runtime/ResonanceAudio/Source/ResonanceAudio/Private/ResonanceAudioLibrary/resonance_audio/utils/wav_reader.h
--- a/Engine/Plugins/Runtime/ResonanceAudio/Source/ResonanceAudio/Private/ResonanceAudioLibrary/resonance_audio/utils/wav_reader.h
+++ b/Engine/Plugins/Runtime/ResonanceAudio/Source/ResonanceAudio/Private/ResonanceAudioLibrary/resonance_audio/utils/wav_reader.h
@@ -19,8 +19,10 @@ limitations under the License.
 
 #include "HAL/PlatformCrt.h"
 
+THIRD_PARTY_INCLUDES_START
 #include <cstdint>
 #include <istream>
+THIRD_PARTY_INCLUDES_END
 
 #include "base/integral_types.h"
 
diff --git a/Engine/Plugins/Runtime/ResonanceAudio/Source/ResonanceAudio/Public/IResonanceAudioModule.h b/Engine/Plugins/Runtime/ResonanceAudio/Source/ResonanceAudio/Public/IResonanceAudioModule.h
--- a/Engine/Plugins/Runtime/ResonanceAudio/Source/ResonanceAudio/Public/IResonanceAudioModule.h
+++ b/Engine/Plugins/Runtime/ResonanceAudio/Source/ResonanceAudio/Public/IResonanceAudioModule.h
@@ -4,6 +4,9 @@
 
 #pragma once
 
+#include "Modules/ModuleInterface.h"
+#include "Modules/ModuleManager.h"
+
 /**
 * The public interface to this module.  In most cases, this interface is only
 * public to sibling modules within this plugin.
diff --git a/Engine/Plugins/Runtime/ResonanceAudio/Source/ResonanceAudio/Public/ResonanceAudioAmbisonicsSettings.h b/Engine/Plugins/Runtime/ResonanceAudio/Source/ResonanceAudio/Public/ResonanceAudioAmbisonicsSettings.h
--- a/Engine/Plugins/Runtime/ResonanceAudio/Source/ResonanceAudio/Public/ResonanceAudioAmbisonicsSettings.h
+++ b/Engine/Plugins/Runtime/ResonanceAudio/Source/ResonanceAudio/Public/ResonanceAudioAmbisonicsSettings.h
@@ -4,12 +4,10 @@
 
 #pragma once
 
-#include "AudioDevice.h"
 #include "ISoundfieldFormat.h"
-#include "ResonanceAudioCommon.h"
-#include "HAL/LowLevelMemTracker.h"
-#include "ResonanceAudioAmbisonicsSettings.generated.h"
+#include "Templates/UniquePtr.h"
 
+#include "ResonanceAudioAmbisonicsSettings.generated.h"
 
 UENUM()
 enum class EResonanceRenderMode : uint8
@@ -34,27 +32,6 @@ enum class EResonanceRenderMode : uint8
 	RoomEffectsOnly,
 };
 
-class FResonanceAmbisonicsSettingsProxy : public ISoundfieldEncodingSettingsProxy
-{
-public:
-	vraudio::RenderingMode RenderingMode;
-
-
-	virtual uint32 GetUniqueId() const override
-	{
-		return static_cast<uint32>(RenderingMode);
-	}
-
-
-	virtual TUniquePtr<ISoundfieldEncodingSettingsProxy> Duplicate() const override
-	{
-		LLM_SCOPE_BYTAG(AudioSpatializationPlugins);
-		FResonanceAmbisonicsSettingsProxy* Proxy = new FResonanceAmbisonicsSettingsProxy();
-		Proxy->RenderingMode = RenderingMode;
-		return TUniquePtr<ISoundfieldEncodingSettingsProxy>(Proxy);
-	}
-};
-
 UCLASS()
 class RESONANCEAUDIO_API UResonanceAudioSoundfieldSettings : public USoundfieldEncodingSettingsBase
 {
@@ -65,45 +42,6 @@ public:
 	UPROPERTY(EditAnywhere, BlueprintReadOnly, Category = Ambisonics)
 	EResonanceRenderMode RenderMode = EResonanceRenderMode::BinauralHighQuality;
 
-	virtual TUniquePtr<ISoundfieldEncodingSettingsProxy> GetProxy() const override
-	{
-		LLM_SCOPE_BYTAG(AudioSpatializationPlugins);
-		FResonanceAmbisonicsSettingsProxy* Proxy = new FResonanceAmbisonicsSettingsProxy();
-
-		switch (RenderMode)
-		{
-			case EResonanceRenderMode::StereoPanning:
-			{
-				Proxy->RenderingMode = vraudio::kStereoPanning;
-				break;
-			}
-			case EResonanceRenderMode::BinauralLowQuality:
-			{
-				Proxy->RenderingMode = vraudio::kBinauralLowQuality;
-				break;
-			}
-			case EResonanceRenderMode::BinauralMediumQuality:
-			{
-				Proxy->RenderingMode = vraudio::kBinauralMediumQuality;
-				break;
-			}
-			case EResonanceRenderMode::BinauralHighQuality:
-			{
-				Proxy->RenderingMode = vraudio::kBinauralHighQuality;
-				break;
-			}
-			case EResonanceRenderMode::RoomEffectsOnly:
-			{
-				Proxy->RenderingMode = vraudio::kRoomEffectsOnly;
-				break;
-			}
-			default:
-			{
-				checkNoEntry();
-			}
-		}
-
-		return TUniquePtr<ISoundfieldEncodingSettingsProxy>(Proxy);
-	}
-
+	virtual TUniquePtr<ISoundfieldEncodingSettingsProxy> GetProxy() const override;
 };
+
diff --git a/Engine/Plugins/Runtime/ResonanceAudio/Source/ResonanceAudio/Public/ResonanceAudioEnums.h b/Engine/Plugins/Runtime/ResonanceAudio/Source/ResonanceAudio/Public/ResonanceAudioEnums.h
new file mode 100644
--- /dev/null
+++ b/Engine/Plugins/Runtime/ResonanceAudio/Source/ResonanceAudio/Public/ResonanceAudioEnums.h
@@ -0,0 +1,75 @@
+//
+// Copyright Google Inc. 2017. All rights reserved.
+//
+
+#pragma once
+
+#include "CoreMinimal.h"
+#include "UObject/ObjectMacros.h"
+
+DECLARE_LOG_CATEGORY_EXTERN(LogResonanceAudio, Display, All);
+
+UENUM(BlueprintType)
+enum class ERaQualityMode : uint8
+{
+	// Stereo panning.
+	STEREO_PANNING						UMETA(DisplayName = "Stereo Panning"),
+	// Binaural Low (First Order Ambisonics).
+	BINAURAL_LOW						UMETA(DisplayName = "Binaural Low Quality"),
+	// Binaural Medium (Second Order Ambisonics).
+	BINAURAL_MEDIUM						UMETA(DisplayName = "Binaural Medium Quality"),
+	// Binaural High (Third Order Ambisonics = Default).
+	BINAURAL_HIGH						UMETA(DisplayName = "Binaural High Quality")
+};
+
+UENUM(BlueprintType)
+enum class ERaSpatializationMethod : uint8
+{
+	// Stereo panning.
+	STEREO_PANNING						UMETA(DisplayName = "Stereo Panning"),
+	// Binaural rendering via HRTF.
+	HRTF								UMETA(DisplayName = "HRTF")
+};
+
+UENUM(BlueprintType)
+enum class ERaDistanceRolloffModel : uint8
+{
+	// Logarithmic distance attenuation model (default).
+	LOGARITHMIC							UMETA(DisplayName = "Logarithmic"),
+	// Linear distance attenuation model.
+	LINEAR								UMETA(DisplayName = "Linear"),
+	// Use Unreal Engine attenuation settings.
+	NONE								UMETA(DisplayName = "None")
+};
+
+UENUM(BlueprintType)
+enum class ERaMaterialName : uint8
+{
+	// Full acoustic energy absorption.
+	TRANSPARENT							UMETA(DisplayName = "Transparent"),
+	ACOUSTIC_CEILING_TILES				UMETA(DisplayName = "Acoustic Ceiling Tiles"),
+	BRICK_BARE							UMETA(DisplayName = "Brick Bare"),
+	BRICK_PAINTED						UMETA(DisplayName = "Brick Painted"),
+	CONCRETE_BLOCK_COARSE				UMETA(DisplayName = "Concrete Block Coarse"),
+	CONCRETE_BLOCK_PAINTED				UMETA(DisplayName = "Concrete Block Painted"),
+	CURTAIN_HEAVY						UMETA(DisplayName = "Curtain Heavy"),
+	FIBER_GLASS_INSULATION				UMETA(DisplayName = "Fiber Glass Insulation"),
+	GLASS_THIN							UMETA(DisplayName = "Glass Thin"),
+	GLASS_THICK							UMETA(DisplayName = "Glass Thick"),
+	GRASS								UMETA(DisplayName = "Grass"),
+	LINOLEUM_ON_CONCRETE				UMETA(DisplayName = "Linoleum On Concrete"),
+	MARBLE								UMETA(DisplayName = "Marble"),
+	METAL								UMETA(DisplayName = "Metal"),
+	PARQUET_ONCONCRETE					UMETA(DisplayName = "Parquet On Concrete"),
+	PLASTER_ROUGH						UMETA(DisplayName = "Plaster Rough"),
+	PLASTER_SMOOTH						UMETA(DisplayName = "Plaster Smooth"),
+	PLYWOOD_PANEL						UMETA(DisplayName = "Plywood Panel"),
+	POLISHED_CONCRETE_OR_TILE			UMETA(DisplayName = "Polished Concrete Or Tile"),
+	SHEETROCK							UMETA(DisplayName = "Sheetrock"),
+	WATER_OR_ICE_SURFACE				UMETA(DisplayName = "Water Or Ice Surface"),
+	WOOD_CEILING						UMETA(DisplayName = "Wood Ceiling"),
+	WOOD_PANEL							UMETA(DisplayName = "Wood Panel"),
+	// Uniform acoustic energy absorption across all frequency bands.
+	UNIFORM								UMETA(DisplayName = "Uniform")
+};
+
diff --git a/Engine/Source/Runtime/AudioCaptureImplementations/AudioCaptureRtAudio/AudioCaptureRtAudio.Build.cs b/Engine/Source/Runtime/AudioCaptureImplementations/AudioCaptureRtAudio/AudioCaptureRtAudio.Build.cs
--- a/Engine/Source/Runtime/AudioCaptureImplementations/AudioCaptureRtAudio/AudioCaptureRtAudio.Build.cs
+++ b/Engine/Source/Runtime/AudioCaptureImplementations/AudioCaptureRtAudio/AudioCaptureRtAudio.Build.cs
@@ -8,6 +8,9 @@ public class AudioCaptureRtAudio : ModuleRules
 
 	public AudioCaptureRtAudio(ReadOnlyTargetRules Target) : base(Target)
 	{
+		// RtAudio throws exceptions for errors.
+		bEnableExceptions = true;
+
 		PrivateDependencyModuleNames.Add("Core");
 		PrivateDependencyModuleNames.Add("AudioCaptureCore");
 
diff --git a/Engine/Source/Runtime/AudioCaptureImplementations/AudioCaptureRtAudio/Private/RtAudio.cpp b/Engine/Source/Runtime/AudioCaptureImplementations/AudioCaptureRtAudio/Private/RtAudio.cpp
--- a/Engine/Source/Runtime/AudioCaptureImplementations/AudioCaptureRtAudio/Private/RtAudio.cpp
+++ b/Engine/Source/Runtime/AudioCaptureImplementations/AudioCaptureRtAudio/Private/RtAudio.cpp
@@ -1,6 +1,7 @@
 // Copyright Epic Games, Inc. All Rights Reserved.
 
 #include "RtAudio.h"
+#include "AudioCaptureCoreLog.h"
 
 #if WITH_RTAUDIO
 
@@ -231,7 +232,7 @@ RtAudio :: RtAudio( RtAudio::Api api )
 
     // No compiled support for specified API value.  Issue a debug
     // warning and continue as if no API was specified.
-    std::cerr << "\nRtAudio: no compiled support for specified API argument!\n" << std::endl;
+    UE_LOG(LogAudioCaptureCore, Warning, TEXT("RtAudio has no compiled support for specified API argument."));
   }
 
   // Iterate through the compiled APIs and return as soon as we find
@@ -248,11 +249,8 @@ RtAudio :: RtAudio( RtAudio::Api api )
   // It should not be possible to get here because the preprocessor
   // definition __RTAUDIO_DUMMY__ is automatically defined if no
   // API-specific definitions are passed to the compiler. But just in
-  // case something weird happens, we'll thow an error.
-  std::string errorText = "\nRtAudio: no compiled API support found ... critical error!!\n\n";
-
-  // Modification: handled exceptions are prohibited in the UE coding standard.
-  // throw_wrapper( RtAudioError( errorText, RtAudioError::UNSPECIFIED ) );
+  // case something weird happens, we'll log an error.
+  UE_LOG(LogAudioCaptureCore, Error, TEXT("RtAudio has no compiled API support found."));
 }
 
 RtAudio :: ~RtAudio()
@@ -10235,4 +10233,4 @@ void RtApi :: byteSwapBuffer( char *buffer, unsigned int samples, RtAudioFormat
 
 #endif // PLATFORM_WINDOWS
 
-#endif // WITH_RTAUDIO
\ No newline at end of file
+#endif // WITH_RTAUDIO
diff --git a/Engine/Source/Runtime/AudioCaptureImplementations/AudioCaptureRtAudio/Private/RtAudioWrapper.cpp b/Engine/Source/Runtime/AudioCaptureImplementations/AudioCaptureRtAudio/Private/RtAudioWrapper.cpp
new file mode 100644
--- /dev/null
+++ b/Engine/Source/Runtime/AudioCaptureImplementations/AudioCaptureRtAudio/Private/RtAudioWrapper.cpp
@@ -0,0 +1,246 @@
+// Copyright Epic Games, Inc. All Rights Reserved.
+
+#include "RtAudioWrapper.h"
+#include "AudioCaptureCoreLog.h"
+
+#if PLATFORM_MICROSOFT || PLATFORM_MAC
+
+#include "RtAudio.h"
+#include "Templates/PimplPtr.h"
+#include "Containers/UnrealString.h"
+
+THIRD_PARTY_INCLUDES_START
+#include <stdexcept>
+#include <string>
+THIRD_PARTY_INCLUDES_END
+
+namespace Audio
+{
+	FRtAudioInputWrapper::FRtAudioInputWrapper()
+	: Impl(MakePimpl<RtAudio>())
+	{
+	}
+
+	// Returns ID of default input device
+	uint32 FRtAudioInputWrapper::GetDefaultInputDevice()
+	{
+		static const TCHAR* Action = TEXT("get default input device");
+		try
+		{
+			return Impl->getDefaultInputDevice();
+		}
+		catch (const std::exception& e)
+		{
+			FString Message(e.what());
+			UE_LOG(LogAudioCaptureCore, Error, TEXT("Failed to %s: %s."), Action, *Message);
+		}
+		catch (...)
+		{
+			UE_LOG(LogAudioCaptureCore, Error, TEXT("Failed to %s."), Action);
+		}
+
+		return InvalidDeviceID;
+	}
+
+	// Returns info about device.
+	FRtAudioInputWrapper::FDeviceInfo FRtAudioInputWrapper::GetDeviceInfo(uint32 InDeviceID)
+	{
+		static const TCHAR* Action = TEXT("get device info");
+		try
+		{
+			RtAudio::DeviceInfo RtInfo = Impl->getDeviceInfo(InDeviceID);
+			return FDeviceInfo {static_cast<float>(RtInfo.preferredSampleRate), RtInfo.inputChannels};
+		}
+		catch (const std::exception& e)
+		{
+			FString Message(e.what());
+			UE_LOG(LogAudioCaptureCore, Error, TEXT("Failed to %s: %s."), Action, *Message);
+		}
+		catch (...)
+		{
+			UE_LOG(LogAudioCaptureCore, Error, TEXT("Failed to %s."), Action);
+		}
+		return FDeviceInfo{};
+	}
+
+	// Opens an audio stream. Returns true on success, false on error. 
+	bool FRtAudioInputWrapper::OpenStream(const FStreamParameters& InStreamParams, float InDesiredSampleRate, uint32* InOutDesiredBufferNumFrames, FAudioCallback Callback, void* InUserData)
+	{
+		static const TCHAR* Action = TEXT("open RtAudio stream");
+		try
+		{
+			RtAudio::StreamParameters RtStreamParams;
+			RtStreamParams.deviceId = InStreamParams.DeviceID; 
+			RtStreamParams.nChannels = InStreamParams.NumChannels;
+			RtStreamParams.firstChannel = 0;
+
+			Impl->openStream(nullptr, &RtStreamParams, RTAUDIO_SINT16, InDesiredSampleRate, InOutDesiredBufferNumFrames, Callback, InUserData);
+		}
+		catch (const std::exception& e)
+		{
+			FString Message(e.what());
+			UE_LOG(LogAudioCaptureCore, Error, TEXT("Failed to %s: %s."), Action, *Message);
+			return false;
+		}
+		catch (...)
+		{
+			UE_LOG(LogAudioCaptureCore, Error, TEXT("Failed to %s."), Action);
+			return false;
+		}
+
+		return true;
+	}
+
+	// Starts an open stream
+	void FRtAudioInputWrapper::StartStream()
+	{
+		static const TCHAR* Action = TEXT("start RtAudio stream");
+		try
+		{
+			Impl->startStream();
+		}
+		catch (const std::exception& e)
+		{
+			FString Message(e.what());
+			UE_LOG(LogAudioCaptureCore, Error, TEXT("Failed to %s: %s."), Action, *Message);
+		}
+		catch (...)
+		{
+			UE_LOG(LogAudioCaptureCore, Error, TEXT("Failed to %s."), Action);
+		}
+	}
+
+	// Returns true if the stream is open.
+	bool FRtAudioInputWrapper::IsStreamOpen()
+	{
+		static const TCHAR* Action = TEXT("check if RtAudio stream is open");
+		try
+		{
+			return Impl->isStreamOpen();
+		}
+		catch (const std::exception& e)
+		{
+			FString Message(e.what());
+			UE_LOG(LogAudioCaptureCore, Error, TEXT("Failed to %s: %s."), Action, *Message);
+		}
+		catch (...)
+		{
+			UE_LOG(LogAudioCaptureCore, Error, TEXT("Failed to %s."), Action);
+		}
+
+		return false;
+	}
+
+	// Stops stream, discarding any remaining samples
+	void FRtAudioInputWrapper::AbortStream()
+	{
+		static const TCHAR* Action = TEXT("abort RtAudio stream");
+		try
+		{
+			Impl->abortStream();
+		}
+		catch (const std::exception& e)
+		{
+			FString Message(e.what());
+			UE_LOG(LogAudioCaptureCore, Error, TEXT("Failed to %s: %s."), Action, *Message);
+		}
+		catch (...)
+		{
+			UE_LOG(LogAudioCaptureCore, Error, TEXT("Failed to %s."), Action);
+		}
+	}
+
+	// Stops stream, allowing any remaining samples to be played.
+	void FRtAudioInputWrapper::StopStream()
+	{
+		static const TCHAR* Action = TEXT("stop RtAudio stream");
+		try
+		{
+			Impl->stopStream();
+		}
+		catch (const std::exception& e)
+		{
+			FString Message(e.what());
+			UE_LOG(LogAudioCaptureCore, Error, TEXT("Failed to %s: %s."), Action, *Message);
+		}
+		catch (...)
+		{
+			UE_LOG(LogAudioCaptureCore, Error, TEXT("Failed to %s."), Action);
+		}
+	}
+
+	// Close stream and free associated memory.
+	void FRtAudioInputWrapper::CloseStream()
+	{
+		static const TCHAR* Action = TEXT("close RtAudio stream");
+		try
+		{
+			Impl->closeStream();
+		}
+		catch (const std::exception& e)
+		{
+			FString Message(e.what());
+			UE_LOG(LogAudioCaptureCore, Error, TEXT("Failed to %s: %s."), Action, *Message);
+		}
+		catch (...)
+		{
+			UE_LOG(LogAudioCaptureCore, Error, TEXT("Failed to %s."), Action);
+		}
+	}
+}
+
+#else
+
+namespace Audio
+{
+	FRtAudioInputWrapper::FRtAudioInputWrapper()
+	{
+		UE_LOG(LogAudioCaptureCore, Error, TEXT("RtAudio is not supported on this platform"));
+	}
+
+	// Returns ID of default input device
+	uint32 FRtAudioInputWrapper::GetDefaultInputDevice()
+	{
+		return InvalidDeviceID;
+	}
+
+	// Returns info about device.
+	FRtAudioInputWrapper::FDeviceInfo FRtAudioInputWrapper::GetDeviceInfo(uint32 InDeviceID)
+	{
+		return FDeviceInfo{};
+	}
+
+	// Opens an audio stream. Returns true on success, false on error. 
+	bool FRtAudioInputWrapper::OpenStream(const FStreamParameters& InStreamParams, float InDesiredSampleRate, uint32* InOutDesiredBufferNumFrames, FAudioCallback Callback, void* InUserData)
+	{
+		return false;
+	}
+
+	// Starts an open stream
+	void FRtAudioInputWrapper::StartStream()
+	{
+	}
+
+	// Returns true if the stream is open.
+	bool FRtAudioInputWrapper::IsStreamOpen()
+	{
+		return false;
+	}
+
+	// Stops stream, discarding any remaining samples
+	void FRtAudioInputWrapper::AbortStream()
+	{
+	}
+
+	// Stops stream, allowing any remaining samples to be played.
+	void FRtAudioInputWrapper::StopStream()
+	{
+	}
+
+	// Close stream and free associated memory.
+	void FRtAudioInputWrapper::CloseStream()
+	{
+	}
+}
+
+#endif // #if PLATFORM_MICROSOFT || PLATFORM_MAC
diff --git a/Engine/Source/Runtime/AudioCaptureImplementations/AudioCaptureRtAudio/Public/RtAudioWrapper.h b/Engine/Source/Runtime/AudioCaptureImplementations/AudioCaptureRtAudio/Public/RtAudioWrapper.h
new file mode 100644
--- /dev/null
+++ b/Engine/Source/Runtime/AudioCaptureImplementations/AudioCaptureRtAudio/Public/RtAudioWrapper.h
@@ -0,0 +1,66 @@
+// Copyright Epic Games, Inc. All Rights Reserved.
+
+#pragma once
+
+#include "Templates/PimplPtr.h"
+
+class RtAudio; // Forward declare;
+
+namespace Audio
+{
+	// FRtAudioInputWrapper wraps RtAudio implementation and provides a way to 
+	// capture input audio from an hardware audio device. 
+	class AUDIOCAPTURERTAUDIO_API FRtAudioInputWrapper
+	{
+	public:
+		static constexpr uint32 InvalidDeviceID = INDEX_NONE;
+
+		// Callback function for receiving audio from RtAudio
+		using FAudioCallback = int (*)(void *OutputBuffer, void *InputBuffer, unsigned int NumFrames, double StreamTime, uint32 StreamStatus, void *UserData);
+
+		// Parameters used for opening streams
+		struct FStreamParameters
+		{
+			uint32 DeviceID = 0;
+			uint32 NumChannels = 0;
+		};
+
+		// Parameters describing audio input device.
+		struct FDeviceInfo
+		{
+			float PreferredSampleRate = 0.f;
+			uint32 NumChannels = 0;
+		};
+
+		FRtAudioInputWrapper();
+
+		// Returns ID of default input device
+		uint32 GetDefaultInputDevice();
+
+		// Returns info about device.
+		FDeviceInfo GetDeviceInfo(uint32 InDeviceID);
+
+		// Opens an audio stream. Returns true on success, false on error. 
+		bool OpenStream(const FStreamParameters& InStreamParams, float InDesiredSampleRate, uint32* InOutDesiredBufferNumFrames, FAudioCallback Callback, void* InUserData);
+
+		// Returns true if the stream is open.
+		bool IsStreamOpen();
+
+		// Starts an open stream
+		void StartStream();
+
+		// Stops stream, discarding any remaining samples
+		void AbortStream();
+
+		// Stops stream, allowing any remaining samples to be played.
+		void StopStream();
+
+		// Close stream and free associated memory.
+		void CloseStream();
+	
+	private:
+
+		TPimplPtr<RtAudio> Impl;
+	};
+}
+
