PR #9950: Solving error with audiocapturecomponent in android, Audio Capture Crashes on Android when using Open Level (Contributed by [REDACTED_PERSON])

#android

diff --git a/Engine/Source/Runtime/AudioCaptureImplementations/Android/AudioCaptureAndroid/Private/AudioCaptureAndroid.cpp b/Engine/Source/Runtime/AudioCaptureImplementations/Android/AudioCaptureAndroid/Private/AudioCaptureAndroid.cpp
--- a/Engine/Source/Runtime/AudioCaptureImplementations/Android/AudioCaptureAndroid/Private/AudioCaptureAndroid.cpp
+++ b/Engine/Source/Runtime/AudioCaptureImplementations/Android/AudioCaptureAndroid/Private/AudioCaptureAndroid.cpp
@@ -61,7 +61,7 @@ bool Audio::FAudioCaptureAndroidStream::OpenCaptureStream(const FAudioCaptureDev
 
 bool Audio::FAudioCaptureAndroidStream::CloseStream()
 {
-	if (!InputOboeStream)
+	if (InputOboeStream)
 	{
 		InputOboeStream->close();
 		InputOboeStream.Reset();
