PR #9970: Use of UTF-8 instead of ANSI characters encoding in GetCaptureDeviceInfo to support unicode characters (Contributed by [REDACTED_PERSON])

diff --git a/Engine/Source/Runtime/AudioCaptureImplementations/AudioCaptureRtAudio/Private/AudioCaptureRtAudio.cpp b/Engine/Source/Runtime/AudioCaptureImplementations/AudioCaptureRtAudio/Private/AudioCaptureRtAudio.cpp
--- a/Engine/Source/Runtime/AudioCaptureImplementations/AudioCaptureRtAudio/Private/AudioCaptureRtAudio.cpp
+++ b/Engine/Source/Runtime/AudioCaptureImplementations/AudioCaptureRtAudio/Private/AudioCaptureRtAudio.cpp
@@ -31,8 +31,8 @@ bool Audio::FAudioCaptureRtAudioStream::GetCaptureDeviceInfo(FCaptureDeviceInfo&
 
 	RtAudio::DeviceInfo DeviceInfo = CaptureDevice.getDeviceInfo(InputDeviceId);
 
-	OutInfo.DeviceName = FString(ANSI_TO_TCHAR(DeviceInfo.name.c_str()));
-	OutInfo.DeviceId = FString(ANSI_TO_TCHAR(DeviceInfo.deviceId.c_str()));
+	OutInfo.DeviceName = FString(FUTF8ToTCHAR(DeviceInfo.name.c_str()));
+	OutInfo.DeviceId = FString(FUTF8ToTCHAR(DeviceInfo.deviceId.c_str()));
 	OutInfo.InputChannels = DeviceInfo.inputChannels;
 	OutInfo.PreferredSampleRate = DeviceInfo.preferredSampleRate;
 
