[Audio Capture] WASAPI device support for audio capture. (3 of 3)

diff --git a/Engine/Config/Windows/BaseWindowsEngine.ini b/Engine/Config/Windows/BaseWindowsEngine.ini
--- a/Engine/Config/Windows/BaseWindowsEngine.ini
+++ b/Engine/Config/Windows/BaseWindowsEngine.ini
@@ -1,5 +1,5 @@
 [Audio]
-AudioCaptureModuleName=AudioCaptureRtAudio
+AudioCaptureModuleName=AudioCaptureWasapi
 
 [OnlineSubsystem]
 LocalPlatformName=WIN
diff --git a/Engine/Plugins/Media/AudioCaptureTimecodeProvider/Source/AudioCaptureTimecodeProvider/Private/AudioCaptureTimecodeProvider.cpp b/Engine/Plugins/Media/AudioCaptureTimecodeProvider/Source/AudioCaptureTimecodeProvider/Private/AudioCaptureTimecodeProvider.cpp
--- a/Engine/Plugins/Media/AudioCaptureTimecodeProvider/Source/AudioCaptureTimecodeProvider/Private/AudioCaptureTimecodeProvider.cpp
+++ b/Engine/Plugins/Media/AudioCaptureTimecodeProvider/Source/AudioCaptureTimecodeProvider/Private/AudioCaptureTimecodeProvider.cpp
@@ -37,14 +37,14 @@ public:
 		//We want a fast timecode detection but we don't want to be called too often.
 		const int32 NumberCaptureFrames = 64;
 
-		Audio::FOnCaptureFunction OnCapture = [this](const float* AudioData, int32 NumFrames, int32 NumChannels, int32 SampleRate, double StreamTime, bool bOverFlow)
+		Audio::FOnAudioCaptureFunction OnCapture = [this](const void* AudioData, int32 NumFrames, int32 NumChannels, int32 SampleRate, double StreamTime, bool bOverFlow)
 		{
-			OnAudioCapture(AudioData, NumFrames, NumChannels, StreamTime, bOverFlow);
+			OnAudioCapture((const float*)AudioData, NumFrames, NumChannels, StreamTime, bOverFlow);
 		};
 
 		Audio::FAudioCaptureDeviceParams Params = Audio::FAudioCaptureDeviceParams();
 
-		if (!AudioCapture.OpenCaptureStream(Params, MoveTemp(OnCapture), NumberCaptureFrames))
+		if (!AudioCapture.OpenAudioCaptureStream(Params, MoveTemp(OnCapture), NumberCaptureFrames))
 		{
 			UE_LOG(LogAudioCaptureTimecodeProvider, Error, TEXT("Can't open the default capture stream for %s."), *Owner->GetName());
 			return false;
diff --git a/Engine/Plugins/Runtime/AudioCapture/Source/AudioCapture/AudioCapture.Build.cs b/Engine/Plugins/Runtime/AudioCapture/Source/AudioCapture/AudioCapture.Build.cs
--- a/Engine/Plugins/Runtime/AudioCapture/Source/AudioCapture/AudioCapture.Build.cs
+++ b/Engine/Plugins/Runtime/AudioCapture/Source/AudioCapture/AudioCapture.Build.cs
@@ -16,12 +16,16 @@ namespace UnrealBuildTool.Rules
 				}
 			);
 
-            if (Target.Platform.IsInGroup(UnrealPlatformGroup.Windows) ||
-				Target.Platform == UnrealTargetPlatform.Mac)
-            {
-                PrivateDependencyModuleNames.Add("AudioCaptureRtAudio");
-            }
-            else if (Target.Platform == UnrealTargetPlatform.IOS)
+			if (Target.Platform.IsInGroup(UnrealPlatformGroup.Windows))
+			{
+				PrivateDependencyModuleNames.Add("AudioCaptureWasapi");
+				PrivateDependencyModuleNames.Add("AudioCaptureRtAudio");
+			}
+			else if (Target.Platform == UnrealTargetPlatform.Mac)
+			{
+				PrivateDependencyModuleNames.Add("AudioCaptureRtAudio");
+			}
+			else if (Target.Platform == UnrealTargetPlatform.IOS)
             {
                 PrivateDependencyModuleNames.Add("AudioCaptureAudioUnit");
             }
diff --git a/Engine/Plugins/Runtime/AudioCapture/Source/AudioCapture/Private/AudioCapture.cpp b/Engine/Plugins/Runtime/AudioCapture/Source/AudioCapture/Private/AudioCapture.cpp
--- a/Engine/Plugins/Runtime/AudioCapture/Source/AudioCapture/Private/AudioCapture.cpp
+++ b/Engine/Plugins/Runtime/AudioCapture/Source/AudioCapture/Private/AudioCapture.cpp
@@ -21,14 +21,14 @@ bool UAudioCapture::OpenDefaultAudioStream()
 {
 	if (!AudioCapture.IsStreamOpen())
 	{
-		Audio::FOnCaptureFunction OnCapture = [this](const float* AudioData, int32 NumFrames, int32 InNumChannels, int32 InSampleRate, double StreamTime, bool bOverFlow)
+		Audio::FOnAudioCaptureFunction OnCapture = [this](const void* AudioData, int32 NumFrames, int32 InNumChannels, int32 InSampleRate, double StreamTime, bool bOverFlow)
 		{
-			OnGeneratedAudio(AudioData, NumFrames * InNumChannels);
+			OnGeneratedAudio((const float*)AudioData, NumFrames * InNumChannels);
 		};
 
 		// Start the stream here to avoid hitching the audio render thread. 
 		Audio::FAudioCaptureDeviceParams Params;
-		if (AudioCapture.OpenCaptureStream(Params, MoveTemp(OnCapture), 1024))
+		if (AudioCapture.OpenAudioCaptureStream(Params, MoveTemp(OnCapture), 1024))
 		{
 			// If we opened the capture stream succesfully, get the capture device info and initialize the UAudioGenerator
 			Audio::FCaptureDeviceInfo Info;
diff --git a/Engine/Plugins/Runtime/AudioCapture/Source/AudioCaptureEditor/AudioCaptureEditor.build.cs b/Engine/Plugins/Runtime/AudioCapture/Source/AudioCaptureEditor/AudioCaptureEditor.build.cs
--- a/Engine/Plugins/Runtime/AudioCapture/Source/AudioCaptureEditor/AudioCaptureEditor.build.cs
+++ b/Engine/Plugins/Runtime/AudioCapture/Source/AudioCaptureEditor/AudioCaptureEditor.build.cs
@@ -23,7 +23,9 @@ public class AudioCaptureEditor : ModuleRules
 
 		if (Target.Platform == UnrealTargetPlatform.Win64)
 		{
-			// Allow us to use direct sound
+			PrivateDependencyModuleNames.Add("AudioCaptureWasapi");
+			
+			// RtAudio depends on DirectSound
 			AddEngineThirdPartyPrivateStaticDependencies(Target, "DirectSound");
 		}
 	}
diff --git a/Engine/Plugins/Runtime/AudioCapture/Source/AudioCaptureEditor/Private/AudioRecordingManager.cpp b/Engine/Plugins/Runtime/AudioCapture/Source/AudioCaptureEditor/Private/AudioRecordingManager.cpp
--- a/Engine/Plugins/Runtime/AudioCapture/Source/AudioCaptureEditor/Private/AudioRecordingManager.cpp
+++ b/Engine/Plugins/Runtime/AudioCapture/Source/AudioCaptureEditor/Private/AudioRecordingManager.cpp
@@ -1,6 +1,7 @@
 // Copyright Epic Games, Inc. All Rights Reserved.
 
 #include "AudioRecordingManager.h"
+#include "AudioCaptureCore.h"
 #include "AudioCaptureEditor.h"
 #include "Sound/AudioSettings.h"
 #include "AssetRegistry/AssetRegistryModule.h"
@@ -19,21 +20,6 @@ DEFINE_LOG_CATEGORY(LogMicManager);
 namespace Audio
 {
 
-/**
-* Callback Function For the Microphone Capture for RtAudio
-*/
-static int32 OnAudioCaptureCallback(void *OutBuffer, void* InBuffer, uint32 InBufferFrames, double StreamTime, uint32 AudioStreamStatus, void* InUserData)
-{
-	// Check for stream overflow (i.e. we're feeding audio faster than we're consuming)
-	const bool bIsOverflow = 0 != AudioStreamStatus;
-
-	// Cast the user data to the singleton mic recorder
-	FAudioRecordingManager* AudioRecordingManager = (FAudioRecordingManager*)InUserData;
-
-	// Call the mic capture callback function
-	return AudioRecordingManager->OnAudioCapture(InBuffer, InBufferFrames, StreamTime, bIsOverflow);
-}
-
 /**
 * FAudioRecordingManager Implementation
 */
@@ -52,9 +38,9 @@ FAudioRecordingManager::FAudioRecordingManager()
 
 FAudioRecordingManager::~FAudioRecordingManager()
 {
-	if (ADC.IsStreamOpen())
+	if (AudioCapture.IsStreamOpen())
 	{
-		ADC.AbortStream();
+		AudioCapture.AbortStream();
 	}
 }
 
@@ -79,10 +65,10 @@ void FAudioRecordingManager::StartRecording(const FRecordingSettings& InSettings
 	RecordingBlockSize = 1024;
 
 	// If we have a stream open close it (reusing streams can cause a blip of previous recordings audio)
-	if (ADC.IsStreamOpen())
+	if (AudioCapture.IsStreamOpen())
 	{
-		ADC.StopStream();
-		ADC.CloseStream();
+		AudioCapture.StopStream();
+		AudioCapture.CloseStream();
 	}
 
 	UE_LOG(LogMicManager, Log, TEXT("Starting mic recording."));
@@ -94,9 +80,16 @@ void FAudioRecordingManager::StartRecording(const FRecordingSettings& InSettings
 	InputGain = Audio::ConvertToLinear(Settings.GainDb);
 
 	// Get the default mic input device info
-	FRtAudioInputWrapper::FDeviceInfo Info = ADC.GetDeviceInfo(StreamParams.DeviceID);
-	RecordingSampleRate = Info.PreferredSampleRate;
-	NumInputChannels = Info.NumChannels;	
+	FAudioCaptureDeviceParams StreamParams;
+	FCaptureDeviceInfo DeviceInfo;
+	if (!AudioCapture.GetCaptureDeviceInfo(DeviceInfo, StreamParams.DeviceIndex))
+	{
+		UE_LOG(LogMicManager, Error, TEXT("Unable to get default audio capture device."));
+		return;
+	}
+
+	RecordingSampleRate = DeviceInfo.PreferredSampleRate;
+	NumInputChannels = DeviceInfo.InputChannels;
 
 	if (NumInputChannels == 0)
 	{
@@ -135,22 +128,27 @@ void FAudioRecordingManager::StartRecording(const FRecordingSettings& InSettings
 	// Publish to the mic input thread that we're ready to record...
 	bRecording = true;
 
-	StreamParams.DeviceID = ADC.GetDefaultInputDevice(); // Only use the default input device for now
-
-	StreamParams.NumChannels = NumInputChannels;
+	StreamParams.NumInputChannels = NumInputChannels;
+	StreamParams.SampleRate = RecordingSampleRate;
+	StreamParams.PCMAudioEncoding = EPCMAudioEncoding::PCM_16;
 
 	uint32 BufferFrames = FMath::Max(RecordingBlockSize, 256);
 
-	UE_LOG(LogMicManager, Log, TEXT("Initialized mic recording manager at %d hz sample rate, %d channels, and %d Recording Block Size"), RecordingSampleRate, StreamParams.NumChannels, BufferFrames);
+	UE_LOG(LogMicManager, Log, TEXT("Initialized mic recording manager at %.4f hz sample rate, %d channels, and %d Recording Block Size"), RecordingSampleRate, StreamParams.NumInputChannels, BufferFrames);
+
+	FOnAudioCaptureFunction OnAudioCaptureCallback = [this](const void* InBuffer, int32 InNumFrames, int32 InNumChannels, int32 InSampleRate, double InStreamTime, bool bInOverflow)
+	{
+		OnAudioCapture(InBuffer, InNumFrames, InStreamTime, bInOverflow);
+	};
 
-	// RtAudio uses exceptions for error handling... 
 	// Open up new audio stream
-	bool bSuccess = ADC.OpenStream(StreamParams, RecordingSampleRate, &BufferFrames, &OnAudioCaptureCallback, this);
+	bool bSuccess = AudioCapture.OpenAudioCaptureStream(StreamParams, OnAudioCaptureCallback, BufferFrames);
+
 	bError = !bSuccess;
 
 	if (bSuccess)
 	{
-		ADC.StartStream();
+		AudioCapture.StartStream();
 	}
 }
 
@@ -209,8 +207,8 @@ void FAudioRecordingManager::StopRecording(TArray<USoundWave*>& OutSoundWaves)
 	if (bRecording)
 	{
 		bRecording = false;
-		ADC.StopStream();
-		ADC.CloseStream();
+		AudioCapture.StopStream();
+		AudioCapture.CloseStream();
 
 		if (CurrentRecordedPCMData.Num() > 0)
 		{
@@ -422,16 +420,19 @@ void FAudioRecordingManager::StopRecording(TArray<USoundWave*>& OutSoundWaves)
 	}
 }
 
-int32 FAudioRecordingManager::OnAudioCapture(void* InBuffer, uint32 InBufferFrames, double StreamTime, bool bOverflow)
+int32 FAudioRecordingManager::OnAudioCapture(const void* InBuffer, uint32 InBufferFrames, double StreamTime, bool bOverflow)
 {
-	FScopeLock Lock(&CriticalSection);
-
-	if (bRecording)
+	// Check recording state prior to locking to avoid potential race condition when stopping
+	if (bRecording && InBufferFrames > 0)
 	{
-		if (bOverflow) 
+		FScopeLock Lock(&CriticalSection);
+
+		if (bOverflow)
+		{
 			++NumOverflowsDetected;
+		}
 
-		CurrentRecordedPCMData.Append((int16*)InBuffer, InBufferFrames * NumInputChannels);
+		CurrentRecordedPCMData.Append((const int16*)InBuffer, InBufferFrames * NumInputChannels);
 		return 0;
 	}
 
diff --git a/Engine/Plugins/Runtime/AudioCapture/Source/AudioCaptureEditor/Private/AudioRecordingManager.h b/Engine/Plugins/Runtime/AudioCapture/Source/AudioCaptureEditor/Private/AudioRecordingManager.h
--- a/Engine/Plugins/Runtime/AudioCapture/Source/AudioCaptureEditor/Private/AudioRecordingManager.h
+++ b/Engine/Plugins/Runtime/AudioCapture/Source/AudioCaptureEditor/Private/AudioRecordingManager.h
@@ -2,10 +2,9 @@
 
 #pragma once
 
-#include "CoreMinimal.h"
-#include "HAL/ThreadSafeBool.h"
+#include "AudioCaptureCore.h"
 #include "Engine/EngineTypes.h"
-#include "RtAudioWrapper.h"
+#include "HAL/ThreadSafeBool.h"
 
 class USoundWave;
 
@@ -56,7 +55,7 @@ namespace Audio
 		void StopRecording(TArray<USoundWave*>& OutSoundWaves);
 
 		// Called by RtAudio when a new audio buffer is ready to be supplied.
-		int32 OnAudioCapture(void* InBuffer, uint32 InBufferFrames, double StreamTime, bool bOverflow);
+		int32 OnAudioCapture(const void* InBuffer, uint32 InBufferFrames, double StreamTime, bool bOverflow);
 
 	private:
 
@@ -72,11 +71,8 @@ namespace Audio
 		// Copy of the input recording settings
 		FRecordingSettings Settings;
 
-		// RtAudio ADC object -- used to interact with low-level audio device.
-		FRtAudioInputWrapper ADC;
-
-		// Stream parameters to initialize the ADC
-		FRtAudioInputWrapper::FStreamParameters StreamParams;
+		// FAudioCapture object -- used to interact with low-level audio device.
+		FAudioCapture AudioCapture;
 
 		// Critical section used to stop and retrieve finished audio buffers.
 		FCriticalSection CriticalSection;
diff --git a/Engine/Source/Runtime/AudioCaptureCore/Private/AudioCaptureCore.cpp b/Engine/Source/Runtime/AudioCaptureCore/Private/AudioCaptureCore.cpp
--- a/Engine/Source/Runtime/AudioCaptureCore/Private/AudioCaptureCore.cpp
+++ b/Engine/Source/Runtime/AudioCaptureCore/Private/AudioCaptureCore.cpp
@@ -62,6 +62,7 @@ namespace Audio
 		return false;
 	}
 
+	PRAGMA_DISABLE_DEPRECATION_WARNINGS
 	bool FAudioCapture::OpenCaptureStream(const FAudioCaptureDeviceParams& InParams, FOnCaptureFunction OnCapture, uint32 NumFramesDesired)
 	{
 		if (Impl.IsValid())
@@ -71,6 +72,17 @@ namespace Audio
 
 		return false;
 	}
+	PRAGMA_ENABLE_DEPRECATION_WARNINGS
+
+	bool FAudioCapture::OpenAudioCaptureStream(const FAudioCaptureDeviceParams& InParams, FOnAudioCaptureFunction OnCapture, uint32 NumFramesDesired)
+	{
+		if (Impl.IsValid())
+		{
+			return Impl->OpenAudioCaptureStream(InParams, MoveTemp(OnCapture), NumFramesDesired);
+		}
+
+		return false;
+	}
 
 	bool FAudioCapture::CloseStream()
 	{
@@ -182,7 +194,7 @@ namespace Audio
 		bool bSuccess = true;
 		if (!AudioCapture.IsStreamOpen())
 		{
-			FOnCaptureFunction OnCapture = [this](const float* AudioData, int32 NumFrames, int32 NumChannels, int32 SampleRate, double StreamTime, bool bOverFlow)
+			FOnAudioCaptureFunction OnCapture = [this](const void* AudioData, int32 NumFrames, int32 NumChannels, int32 SampleRate, double StreamTime, bool bOverFlow)
 			{
 				int32 NumSamples = NumChannels * NumFrames;
 
@@ -212,7 +224,7 @@ namespace Audio
 			FAudioCaptureDeviceParams Params = FAudioCaptureDeviceParams();
 
 			// Start the stream here to avoid hitching the audio render thread. 
-			if (AudioCapture.OpenCaptureStream(Params, MoveTemp(OnCapture), 1024))
+			if (AudioCapture.OpenAudioCaptureStream(Params, MoveTemp(OnCapture), 1024))
 			{
 				AudioCapture.StartStream();
 			}
diff --git a/Engine/Source/Runtime/AudioCaptureCore/Private/AudioCaptureInternal.h b/Engine/Source/Runtime/AudioCaptureCore/Private/AudioCaptureInternal.h
--- a/Engine/Source/Runtime/AudioCaptureCore/Private/AudioCaptureInternal.h
+++ b/Engine/Source/Runtime/AudioCaptureCore/Private/AudioCaptureInternal.h
@@ -16,7 +16,7 @@ namespace Audio
 
 		// Begin IAudioCaptureStream
 		virtual bool GetCaptureDeviceInfo(FCaptureDeviceInfo& OutInfo, int32 DeviceIndex) override { return false; }
-		virtual bool OpenCaptureStream(const FAudioCaptureDeviceParams& InParams, FOnCaptureFunction OnCapture, uint32 NumFramesDesired) override { return false; }
+		virtual bool OpenAudioCaptureStream(const FAudioCaptureDeviceParams& InParams, FOnAudioCaptureFunction OnCapture, uint32 NumFramesDesired) override { return false; }
 		virtual bool CloseStream() override { return false; }
 		virtual bool StartStream() override { return false; }
 		virtual bool StopStream() override { return false; }
diff --git a/Engine/Source/Runtime/AudioCaptureCore/Public/AudioCaptureCore.h b/Engine/Source/Runtime/AudioCaptureCore/Public/AudioCaptureCore.h
--- a/Engine/Source/Runtime/AudioCaptureCore/Public/AudioCaptureCore.h
+++ b/Engine/Source/Runtime/AudioCaptureCore/Public/AudioCaptureCore.h
@@ -43,7 +43,13 @@ namespace Audio
 		bool GetCaptureDeviceInfo(FCaptureDeviceInfo& OutInfo, int32 DeviceIndex = INDEX_NONE);
 
 		// Opens the audio capture stream with the given parameters
+		PRAGMA_DISABLE_DEPRECATION_WARNINGS
+		UE_DEPRECATED(5.3, "OpenCaptureStream is deprecated, please use OpenAudioCaptureStream instead.")
 		bool OpenCaptureStream(const FAudioCaptureDeviceParams& InParams, FOnCaptureFunction OnCapture, uint32 NumFramesDesired);
+		PRAGMA_ENABLE_DEPRECATION_WARNINGS
+
+		// Opens the audio capture stream with the given parameters
+		bool OpenAudioCaptureStream(const FAudioCaptureDeviceParams& InParams, FOnAudioCaptureFunction OnCapture, uint32 NumFramesDesired);
 
 		// Closes the audio capture stream
 		bool CloseStream();
diff --git a/Engine/Source/Runtime/AudioCaptureCore/Public/AudioCaptureDeviceInterface.h b/Engine/Source/Runtime/AudioCaptureCore/Public/AudioCaptureDeviceInterface.h
--- a/Engine/Source/Runtime/AudioCaptureCore/Public/AudioCaptureDeviceInterface.h
+++ b/Engine/Source/Runtime/AudioCaptureCore/Public/AudioCaptureDeviceInterface.h
@@ -9,6 +9,21 @@ namespace Audio
 {
 	enum class EHardwareInputFeature : uint8;
 	
+	/**
+	 * PCM audio format (bits per sample) supported by device layers (WASAPI, CoreAudio, etc.)
+	 */
+	enum class EPCMAudioEncoding : uint8
+	{
+		UNKNOWN,
+		PCM_8,
+		PCM_16,
+		PCM_24,
+		PCM_24_IN_32,
+		PCM_32,
+		FLOATING_POINT_32,
+		FLOATING_POINT_64
+	};
+
 	struct FCaptureDeviceInfo
  	{
 		FString DeviceName;
@@ -18,24 +33,34 @@ namespace Audio
 		bool bSupportsHardwareAEC;
 	};
 
-	const int32 DefaultDeviceIndex = INDEX_NONE;
+	// Indicates the default, system capture device should be used for capturing audio
+	constexpr int32 DefaultDeviceIndex = INDEX_NONE;
+	// Indicates invalid sample rate and that the default device's sample rate should be used
+	constexpr int32 InvalidDeviceSampleRate = -1;
+	// Indicates invalid channel count and that the input channel count of the default device should be used
+	constexpr int32 InvalidDeviceChannelCount = -1;
+	// Some clients such as UAudioCaptureComponent assume 32-bit float bit depth
+	constexpr EPCMAudioEncoding DefaultDeviceEncoding = EPCMAudioEncoding::FLOATING_POINT_32;
 
 	struct FAudioCaptureDeviceParams
 	{
-		// set to true to use this device's built-in echo cancellation, if possible.
-		bool bUseHardwareAEC;
+		// Set to true to use this device's built-in echo cancellation, if possible.
+		bool bUseHardwareAEC = false;
 		// Set this to INDEX_NONE 
-		int32 DeviceIndex;
-
-		FAudioCaptureDeviceParams()
-			: bUseHardwareAEC(false)
-			, DeviceIndex(DefaultDeviceIndex)
-		{
-		}
+		int32 DeviceIndex = DefaultDeviceIndex;
+		// Number of input channels on the device available for recording audio
+		int32 NumInputChannels = InvalidDeviceChannelCount;
+		// The sample rate to use when recording audio on this device
+		int32 SampleRate = InvalidDeviceSampleRate;
+		// The bit depth to use when recording audio on this device
+		EPCMAudioEncoding PCMAudioEncoding = DefaultDeviceEncoding;
 	};
 
+	UE_DEPRECATED(5.3, "FOnCaptureFunction is deprecated, please use FOnAudioCaptureFunction instead.")
 	typedef TFunction<void(const float* InAudio, int32 NumFrames, int32 NumChannels, int32 SampleRate, double StreamTime, bool bOverFlow)> FOnCaptureFunction;
-	
+	// Callable function type used for audio capture callbacks.
+	typedef TFunction<void(const void* InAudio, int32 NumFrames, int32 NumChannels, int32 SampleRate, double StreamTime, bool bOverFlow)> FOnAudioCaptureFunction;
+
 	class IAudioCaptureStream : public IModularFeature
 	{
 	public:
@@ -50,7 +75,20 @@ namespace Audio
 		// Returns the audio capture device information at the given Id.
 		virtual bool GetCaptureDeviceInfo(FCaptureDeviceInfo& OutInfo, int32 DeviceIndex) = 0;
 		// Opens the audio capture stream with the given parameters
-		virtual bool OpenCaptureStream(const FAudioCaptureDeviceParams& InParams, FOnCaptureFunction InOnCapture, uint32 NumFramesDesired) = 0;
+		PRAGMA_DISABLE_DEPRECATION_WARNINGS
+		UE_DEPRECATED(5.3, "OpenCaptureStream is deprecated, please use OpenAudioCaptureStream instead.")
+		virtual bool OpenCaptureStream(const FAudioCaptureDeviceParams& InParams, FOnCaptureFunction InOnCapture, uint32 NumFramesDesired)
+		{
+			FOnAudioCaptureFunction OnCaptureCallback = [InOnCapture](const void* InBuffer, int32 InNumFrames, int32 InNumChannels, int32 InSampleRate, double InStreamTime, bool bInOverflow)
+			{
+				InOnCapture((const float*)InBuffer, InNumFrames, InNumChannels, InSampleRate, InStreamTime, bInOverflow);
+			};
+
+			return OpenAudioCaptureStream(InParams, OnCaptureCallback, NumFramesDesired);
+		}
+		PRAGMA_ENABLE_DEPRECATION_WARNINGS
+		// Opens the audio capture stream with the given parameters and using FOnAudioCaptureFunction for callbacks
+		virtual bool OpenAudioCaptureStream(const FAudioCaptureDeviceParams& InParams, FOnAudioCaptureFunction InOnCapture, uint32 NumFramesDesired) = 0;
 		// Closes the audio capture stream
 		virtual bool CloseStream() = 0;
 		// Start the audio capture stream
diff --git a/Engine/Source/Runtime/AudioCaptureImplementations/Android/AudioCaptureAndroid/Private/AudioCaptureAndroid.cpp b/Engine/Source/Runtime/AudioCaptureImplementations/Android/AudioCaptureAndroid/Private/AudioCaptureAndroid.cpp
--- a/Engine/Source/Runtime/AudioCaptureImplementations/Android/AudioCaptureAndroid/Private/AudioCaptureAndroid.cpp
+++ b/Engine/Source/Runtime/AudioCaptureImplementations/Android/AudioCaptureAndroid/Private/AudioCaptureAndroid.cpp
@@ -28,7 +28,7 @@ bool Audio::FAudioCaptureAndroidStream::GetCaptureDeviceInfo(FCaptureDeviceInfo&
 	return true;
 }
 
-bool Audio::FAudioCaptureAndroidStream::OpenCaptureStream(const FAudioCaptureDeviceParams& InParams, FOnCaptureFunction InOnCapture, uint32 NumFramesDesired)
+bool Audio::FAudioCaptureAndroidStream::OpenAudioCaptureStream(const FAudioCaptureDeviceParams& InParams, FOnAudioCaptureFunction InOnCapture, uint32 NumFramesDesired)
 {
 	// Build stream settings object.
 	oboe::AudioStreamBuilder StreamBuilder;
diff --git a/Engine/Source/Runtime/AudioCaptureImplementations/Android/AudioCaptureAndroid/Private/AudioCaptureAndroid.h b/Engine/Source/Runtime/AudioCaptureImplementations/Android/AudioCaptureAndroid/Private/AudioCaptureAndroid.h
--- a/Engine/Source/Runtime/AudioCaptureImplementations/Android/AudioCaptureAndroid/Private/AudioCaptureAndroid.h
+++ b/Engine/Source/Runtime/AudioCaptureImplementations/Android/AudioCaptureAndroid/Private/AudioCaptureAndroid.h
@@ -12,7 +12,7 @@ namespace Audio
 		FAudioCaptureAndroidStream();
 
 		virtual bool GetCaptureDeviceInfo(FCaptureDeviceInfo& OutInfo, int32 DeviceIndex) override;
-		virtual bool OpenCaptureStream(const FAudioCaptureDeviceParams& InParams, FOnCaptureFunction InOnCapture, uint32 NumFramesDesired) override;
+		virtual bool OpenAudioCaptureStream(const FAudioCaptureDeviceParams& InParams, FOnAudioCaptureFunction InOnCapture, uint32 NumFramesDesired) override;
 		virtual bool CloseStream() override;
 		virtual bool StartStream() override;
 		virtual bool StopStream() override;
@@ -29,7 +29,7 @@ namespace Audio
 		// end oboe::AudioStreamCallback
 
 	private:
-		FOnCaptureFunction OnCapture;
+		FOnAudioCaptureFunction OnCapture;
 		int32 NumChannels;
 		int32 SampleRate;
 
diff --git a/Engine/Source/Runtime/AudioCaptureImplementations/AudioCaptureRtAudio/Private/AudioCaptureRtAudio.cpp b/Engine/Source/Runtime/AudioCaptureImplementations/AudioCaptureRtAudio/Private/AudioCaptureRtAudio.cpp
--- a/Engine/Source/Runtime/AudioCaptureImplementations/AudioCaptureRtAudio/Private/AudioCaptureRtAudio.cpp
+++ b/Engine/Source/Runtime/AudioCaptureImplementations/AudioCaptureRtAudio/Private/AudioCaptureRtAudio.cpp
@@ -1,8 +1,53 @@
 // Copyright Epic Games, Inc. All Rights Reserved.
 
 #include "AudioCaptureRtAudio.h"
+#include "AudioCaptureCoreLog.h"
+#include "AudioCaptureDeviceInterface.h"
 
 
+static constexpr uint32 InvalidDeviceID = INDEX_NONE;
+
+#if WITH_RTAUDIO
+bool GetAudioFormatFromEncoding(Audio::EPCMAudioEncoding InEncoding, RtAudioFormat& OutAudioFormat)
+{
+	switch (InEncoding)
+	{
+	case Audio::EPCMAudioEncoding::PCM_8:
+		OutAudioFormat = RTAUDIO_SINT8;
+		break;
+
+	case Audio::EPCMAudioEncoding::PCM_16:
+		OutAudioFormat = RTAUDIO_SINT16;
+		break;
+
+	case Audio::EPCMAudioEncoding::PCM_24:
+		OutAudioFormat = RTAUDIO_SINT24;
+		break;
+
+	case Audio::EPCMAudioEncoding::PCM_32:
+		OutAudioFormat = RTAUDIO_SINT32;
+		break;
+
+		// Default to 32-bit float
+	case Audio::EPCMAudioEncoding::UNKNOWN:
+	case Audio::EPCMAudioEncoding::FLOATING_POINT_32:
+		OutAudioFormat = RTAUDIO_FLOAT32;
+		break;
+
+	case Audio::EPCMAudioEncoding::FLOATING_POINT_64:
+		OutAudioFormat = RTAUDIO_FLOAT64;
+		break;
+
+		// RtAudio does not support 24-bit audio in a 32-bit container
+	case Audio::EPCMAudioEncoding::PCM_24_IN_32:
+	default:
+		return false;
+	}
+
+	return true;
+}
+#endif // WITH_RTAUDIO
+
 Audio::FAudioCaptureRtAudioStream::FAudioCaptureRtAudioStream()
 	: NumChannels(0)
 	, SampleRate(0)
@@ -23,18 +68,33 @@ static int32 OnAudioCaptureCallback(void *OutBuffer, void* InBuffer, uint32 InBu
 bool Audio::FAudioCaptureRtAudioStream::GetCaptureDeviceInfo(FCaptureDeviceInfo& OutInfo, int32 DeviceIndex)
 {
 #if WITH_RTAUDIO
-	uint32 InputDeviceId = CaptureDevice.getDefaultInputDevice();
-	if (DeviceIndex != DefaultDeviceIndex)
+	static const TCHAR* Action = TEXT("get device info");
+	try
 	{
-		InputDeviceId = DeviceIndex;
-	}
-
-	RtAudio::DeviceInfo DeviceInfo = CaptureDevice.getDeviceInfo(InputDeviceId);
+		uint32 InputDeviceId = GetDefaultInputDevice();
+		if (DeviceIndex != DefaultDeviceIndex)
+		{
+			InputDeviceId = DeviceIndex;
+		}
 
-	OutInfo.DeviceName = FString(FUTF8ToTCHAR(DeviceInfo.name.c_str()));
-	OutInfo.DeviceId = FString(FUTF8ToTCHAR(DeviceInfo.deviceId.c_str()));
-	OutInfo.InputChannels = DeviceInfo.inputChannels;
-	OutInfo.PreferredSampleRate = DeviceInfo.preferredSampleRate;
+		RtAudio::DeviceInfo DeviceInfo = CaptureDevice.getDeviceInfo(InputDeviceId);
+		
+		OutInfo.DeviceName = FString(FUTF8ToTCHAR(DeviceInfo.name.c_str()));
+		OutInfo.DeviceId = FString(FUTF8ToTCHAR(DeviceInfo.deviceId.c_str()));
+		OutInfo.InputChannels = DeviceInfo.inputChannels;
+		OutInfo.PreferredSampleRate = DeviceInfo.preferredSampleRate;
+	}
+	catch (const std::exception& e)
+	{
+		FString Message(e.what());
+		UE_LOG(LogAudioCaptureCore, Error, TEXT("Failed to %s: %s."), Action, *Message);
+		return false;
+	}
+	catch (...)
+	{
+		UE_LOG(LogAudioCaptureCore, Error, TEXT("Failed to %s."), Action);
+		return false;
+	}
 
 	return true;
 #else // ^^^ WITH_RTAUDIO vvv !WITH_RTAUDIO
@@ -42,67 +102,105 @@ bool Audio::FAudioCaptureRtAudioStream::GetCaptureDeviceInfo(FCaptureDeviceInfo&
 #endif // !WITH_RTAUDIO
 }
 
-bool Audio::FAudioCaptureRtAudioStream::OpenCaptureStream(const FAudioCaptureDeviceParams& InParams, FOnCaptureFunction InOnCapture, uint32 NumFramesDesired)
+bool Audio::FAudioCaptureRtAudioStream::OpenAudioCaptureStream(const FAudioCaptureDeviceParams& InParams, FOnAudioCaptureFunction InOnCapture, uint32 NumFramesDesired)
 {
 #if WITH_RTAUDIO
-	uint32 InputDeviceId = CaptureDevice.getDefaultInputDevice();
-	if (InParams.DeviceIndex != DefaultDeviceIndex)
+	static const TCHAR* Action = TEXT("open RtAudio stream");
+	try
 	{
-		// InParams.DeviceIndex is the index inside the sublist of devices including only
-		// the Input devices.
-		uint32 NumInputDevices = 0;
-		bool DeviceFound = false;
-		uint32 NumDevices = CaptureDevice.getDeviceCount();
-		for (uint32 DeviceIndex = 0; DeviceIndex < NumDevices; DeviceIndex++)
+		uint32 InputDeviceId = GetDefaultInputDevice();
+		if (InParams.DeviceIndex != DefaultDeviceIndex)
 		{
-			FCaptureDeviceInfo DeviceInfo;
-			GetCaptureDeviceInfo(DeviceInfo, DeviceIndex);
-			bool IsInput = DeviceInfo.InputChannels > 0;
-			if (!IsInput)
+			// InParams.DeviceIndex is the index inside the sublist of devices including only
+			// the Input devices.
+			uint32 NumInputDevices = 0;
+			bool DeviceFound = false;
+			uint32 NumDevices = CaptureDevice.getDeviceCount();
+			for (uint32 DeviceIndex = 0; DeviceIndex < NumDevices; DeviceIndex++)
 			{
-				continue;
+				FCaptureDeviceInfo DeviceInfo;
+				GetCaptureDeviceInfo(DeviceInfo, DeviceIndex);
+				bool IsInput = DeviceInfo.InputChannels > 0;
+				if (!IsInput)
+				{
+					continue;
+				}
+				if (NumInputDevices == InParams.DeviceIndex)
+				{
+					DeviceFound = true;
+					InputDeviceId = DeviceIndex;
+					break;
+				}
+				NumInputDevices++;
 			}
-			if (NumInputDevices == InParams.DeviceIndex)
+			if (!DeviceFound)
 			{
-				DeviceFound = true;
-				InputDeviceId = DeviceIndex;
-				break;
+				return false;
 			}
-			NumInputDevices++;
 		}
-		if (!DeviceFound)
+
+		RtAudio::DeviceInfo DeviceInfo = CaptureDevice.getDeviceInfo(InputDeviceId);
+
+		RtAudio::StreamParameters RtAudioStreamParams;
+		RtAudioStreamParams.deviceId = InputDeviceId;
+		RtAudioStreamParams.firstChannel = 0;
+
+		if (InParams.NumInputChannels == InvalidDeviceChannelCount)
 		{
-			return false;
+			RtAudioStreamParams.nChannels = FMath::Min((int32)DeviceInfo.inputChannels, 8);
+		}
+		else
+		{
+			RtAudioStreamParams.nChannels = InParams.NumInputChannels;
 		}
-	}
-
-	RtAudio::DeviceInfo DeviceInfo = CaptureDevice.getDeviceInfo(InputDeviceId);
 
-	RtAudio::StreamParameters RtAudioStreamParams;
-	RtAudioStreamParams.deviceId = InputDeviceId;
-	RtAudioStreamParams.firstChannel = 0;
-	RtAudioStreamParams.nChannels = FMath::Min((int32)DeviceInfo.inputChannels, 8);
+		if (CaptureDevice.isStreamOpen())
+		{
+			CaptureDevice.stopStream();
+			CaptureDevice.closeStream();
+		}
 
-	if (CaptureDevice.isStreamOpen())
-	{
-		CaptureDevice.stopStream();
-		CaptureDevice.closeStream();
-	}
+		uint32 NumFrames = NumFramesDesired;
+		NumChannels = RtAudioStreamParams.nChannels;
+		OnCapture = MoveTemp(InOnCapture);
 
-	uint32 NumFrames = NumFramesDesired;
-	NumChannels = RtAudioStreamParams.nChannels;
-	SampleRate = DeviceInfo.preferredSampleRate;
-	OnCapture = MoveTemp(InOnCapture);
+		if (InParams.SampleRate == InvalidDeviceSampleRate)
+		{
+			SampleRate = DeviceInfo.preferredSampleRate;
+		}
+		else
+		{
+			SampleRate = InParams.SampleRate;
+		}
 
-	// Open up new audio stream
-	CaptureDevice.openStream(nullptr, &RtAudioStreamParams, RTAUDIO_FLOAT32, SampleRate, &NumFrames, &OnAudioCaptureCallback, this);
+		RtAudioFormat AudioFormat;
+		if (GetAudioFormatFromEncoding(InParams.PCMAudioEncoding, AudioFormat))
+		{
+			// Open up new audio stream
+			CaptureDevice.openStream(nullptr, &RtAudioStreamParams, AudioFormat, SampleRate, &NumFrames, &OnAudioCaptureCallback, this);
+			if (!CaptureDevice.isStreamOpen())
+			{
+				return false;
+			}
+		}
+		else
+		{
+			return false;
+		}
 
-	if (!CaptureDevice.isStreamOpen())
+		SampleRate = CaptureDevice.getStreamSampleRate();
+	}
+	catch (const std::exception& e)
 	{
+		FString Message(e.what());
+		UE_LOG(LogAudioCaptureCore, Error, TEXT("Failed to %s: %s."), Action, *Message);
+		return false;
+	}
+	catch (...)
+	{
+		UE_LOG(LogAudioCaptureCore, Error, TEXT("Failed to %s."), Action);
 		return false;
 	}
-
-	SampleRate = CaptureDevice.getStreamSampleRate();
 
 	return true;
 #else // ^^^ WITH_RTAUDIO vvv !WITH_RTAUDIO
@@ -113,11 +211,26 @@ bool Audio::FAudioCaptureRtAudioStream::OpenCaptureStream(const FAudioCaptureDev
 bool Audio::FAudioCaptureRtAudioStream::CloseStream()
 {
 #if WITH_RTAUDIO
-	if (CaptureDevice.isStreamOpen())
+	static const TCHAR* Action = TEXT("close RtAudio stream");
+	try
+	{
+		if (CaptureDevice.isStreamOpen())
+		{
+			CaptureDevice.closeStream();
+			return true;
+		}
+	}
+	catch (const std::exception& e)
 	{
-		CaptureDevice.closeStream();
+		FString Message(e.what());
+		UE_LOG(LogAudioCaptureCore, Error, TEXT("Failed to %s: %s."), Action, *Message);
 	}
-	return true;
+	catch (...)
+	{
+		UE_LOG(LogAudioCaptureCore, Error, TEXT("Failed to %s."), Action);
+	}
+
+	return false;
 #else // ^^^ WITH_RTAUDIO vvv !WITH_RTAUDIO
 	return false;
 #endif // !WITH_RTAUDIO
@@ -126,8 +239,23 @@ bool Audio::FAudioCaptureRtAudioStream::CloseStream()
 bool Audio::FAudioCaptureRtAudioStream::StartStream()
 {
 #if WITH_RTAUDIO
-	CaptureDevice.startStream();
-	return true;
+	static const TCHAR* Action = TEXT("start RtAudio stream");
+	try
+	{
+		CaptureDevice.startStream();
+		return true;
+	}
+	catch (const std::exception& e)
+	{
+		FString Message(e.what());
+		UE_LOG(LogAudioCaptureCore, Error, TEXT("Failed to %s: %s."), Action, *Message);
+	}
+	catch (...)
+	{
+		UE_LOG(LogAudioCaptureCore, Error, TEXT("Failed to %s."), Action);
+	}
+
+	return false;
 #else // ^^^ WITH_RTAUDIO vvv !WITH_RTAUDIO
 	return false;
 #endif // !WITH_RTAUDIO
@@ -136,11 +264,26 @@ bool Audio::FAudioCaptureRtAudioStream::StartStream()
 bool Audio::FAudioCaptureRtAudioStream::StopStream()
 {
 #if WITH_RTAUDIO
-	if (CaptureDevice.isStreamOpen())
+	static const TCHAR* Action = TEXT("stop RtAudio stream");
+	try
+	{
+		if (CaptureDevice.isStreamOpen())
+		{
+			CaptureDevice.stopStream();
+			return true;
+		}
+	}
+	catch (const std::exception& e)
 	{
-		CaptureDevice.stopStream();
+		FString Message(e.what());
+		UE_LOG(LogAudioCaptureCore, Error, TEXT("Failed to %s: %s."), Action, *Message);
 	}
-	return true;
+	catch (...)
+	{
+		UE_LOG(LogAudioCaptureCore, Error, TEXT("Failed to %s."), Action);
+	}
+
+	return false;
 #else // ^^^ WITH_RTAUDIO vvv !WITH_RTAUDIO
 	return false;
 #endif // !WITH_RTAUDIO
@@ -149,11 +292,26 @@ bool Audio::FAudioCaptureRtAudioStream::StopStream()
 bool Audio::FAudioCaptureRtAudioStream::AbortStream()
 {
 #if WITH_RTAUDIO
-	if (CaptureDevice.isStreamOpen())
+	static const TCHAR* Action = TEXT("abort RtAudio stream");
+	try
 	{
-		CaptureDevice.abortStream();
+		if (CaptureDevice.isStreamOpen())
+		{
+			CaptureDevice.abortStream();
+			return true;
+		}
 	}
-	return true;
+	catch (const std::exception& e)
+	{
+		FString Message(e.what());
+		UE_LOG(LogAudioCaptureCore, Error, TEXT("Failed to %s: %s."), Action, *Message);
+	}
+	catch (...)
+	{
+		UE_LOG(LogAudioCaptureCore, Error, TEXT("Failed to %s."), Action);
+	}
+
+	return false;
 #else // ^^^ WITH_RTAUDIO vvv !WITH_RTAUDIO
 	return false;
 #endif // !WITH_RTAUDIO
@@ -162,8 +320,23 @@ bool Audio::FAudioCaptureRtAudioStream::AbortStream()
 bool Audio::FAudioCaptureRtAudioStream::GetStreamTime(double& OutStreamTime)
 {
 #if WITH_RTAUDIO
-	OutStreamTime = CaptureDevice.getStreamTime();
-	return true;
+	static const TCHAR* Action = TEXT("get RtAudio stream time");
+	try
+	{
+		OutStreamTime = CaptureDevice.getStreamTime();
+		return true;
+	}
+	catch (const std::exception& e)
+	{
+		FString Message(e.what());
+		UE_LOG(LogAudioCaptureCore, Error, TEXT("Failed to %s: %s."), Action, *Message);
+	}
+	catch (...)
+	{
+		UE_LOG(LogAudioCaptureCore, Error, TEXT("Failed to %s."), Action);
+	}
+
+	return false;
 #else // ^^^ WITH_RTAUDIO vvv !WITH_RTAUDIO
 	return false;
 #endif // !WITH_RTAUDIO
@@ -172,7 +345,22 @@ bool Audio::FAudioCaptureRtAudioStream::GetStreamTime(double& OutStreamTime)
 bool Audio::FAudioCaptureRtAudioStream::IsStreamOpen() const
 {
 #if WITH_RTAUDIO
-	return CaptureDevice.isStreamOpen();
+	static const TCHAR* Action = TEXT("check if RtAudio stream is open");
+	try
+	{
+		return CaptureDevice.isStreamOpen();
+	}
+	catch (const std::exception& e)
+	{
+		FString Message(e.what());
+		UE_LOG(LogAudioCaptureCore, Error, TEXT("Failed to %s: %s."), Action, *Message);
+	}
+	catch (...)
+	{
+		UE_LOG(LogAudioCaptureCore, Error, TEXT("Failed to %s."), Action);
+	}
+
+	return false;
 #else // ^^^ WITH_RTAUDIO vvv !WITH_RTAUDIO
 	return false;
 #endif // !WITH_RTAUDIO
@@ -181,7 +369,22 @@ bool Audio::FAudioCaptureRtAudioStream::IsStreamOpen() const
 bool Audio::FAudioCaptureRtAudioStream::IsCapturing() const
 {
 #if WITH_RTAUDIO
-	return CaptureDevice.isStreamRunning();
+	static const TCHAR* Action = TEXT("check if RtAudio stream is running");
+	try
+	{
+		return CaptureDevice.isStreamRunning();
+	}
+	catch (const std::exception& e)
+	{
+		FString Message(e.what());
+		UE_LOG(LogAudioCaptureCore, Error, TEXT("Failed to %s: %s."), Action, *Message);
+	}
+	catch (...)
+	{
+		UE_LOG(LogAudioCaptureCore, Error, TEXT("Failed to %s."), Action);
+	}
+
+	return false;
 #else // ^^^ WITH_RTAUDIO vvv !WITH_RTAUDIO
 	return false;
 #endif // !WITH_RTAUDIO
@@ -190,30 +393,69 @@ bool Audio::FAudioCaptureRtAudioStream::IsCapturing() const
 void Audio::FAudioCaptureRtAudioStream::OnAudioCapture(void* InBuffer, uint32 InBufferFrames, double StreamTime, bool bOverflow)
 {
 #if WITH_RTAUDIO
-	float* InBufferData = (float*)InBuffer;
-	OnCapture(InBufferData, InBufferFrames, NumChannels, SampleRate, StreamTime, bOverflow);
+	OnCapture(InBuffer, InBufferFrames, NumChannels, SampleRate, StreamTime, bOverflow);
 #endif // WITH_RTAUDIO
 }
 
 bool Audio::FAudioCaptureRtAudioStream::GetInputDevicesAvailable(TArray<FCaptureDeviceInfo>& OutDevices)
 {
 #if WITH_RTAUDIO
-	uint32 NumDevices = CaptureDevice.getDeviceCount();
-	OutDevices.Reset();
-	// Iterate over all devices and include in the output only the ones that are input devices
-	for (uint32 DeviceIndex = 0; DeviceIndex < NumDevices; DeviceIndex++)
-	{
-		FCaptureDeviceInfo DeviceInfo;
-		GetCaptureDeviceInfo(DeviceInfo, DeviceIndex);
-		bool IsInput = DeviceInfo.InputChannels > 0;
-		if (IsInput)
+	static const TCHAR* Action = TEXT("get input devices available RtAudio");
+	try
+	{
+		uint32 NumDevices = CaptureDevice.getDeviceCount();
+		OutDevices.Reset();
+		// Iterate over all devices and include in the output only the ones that are input devices
+		for (uint32 DeviceIndex = 0; DeviceIndex < NumDevices; DeviceIndex++)
 		{
-			OutDevices.Add(DeviceInfo);
+			FCaptureDeviceInfo DeviceInfo;
+			GetCaptureDeviceInfo(DeviceInfo, DeviceIndex);
+			bool IsInput = DeviceInfo.InputChannels > 0;
+			if (IsInput)
+			{
+				OutDevices.Add(DeviceInfo);
+			}
 		}
+
+		return true;
+	}
+	catch (const std::exception& e)
+	{
+		FString Message(e.what());
+		UE_LOG(LogAudioCaptureCore, Error, TEXT("Failed to %s: %s."), Action, *Message);
+	}
+	catch (...)
+	{
+		UE_LOG(LogAudioCaptureCore, Error, TEXT("Failed to %s."), Action);
 	}
 
-	return true;
+	return false;
 #else // ^^^ WITH_RTAUDIO vvv !WITH_RTAUDIO
 	return false;
 #endif // !WITH_RTAUDIO
 }
+
+uint32 Audio::FAudioCaptureRtAudioStream::GetDefaultInputDevice()
+{
+#if WITH_RTAUDIO
+	static const TCHAR* Action = TEXT("get default input device");
+	try
+	{
+		return CaptureDevice.getDefaultInputDevice();
+	}
+	catch (const std::exception& e)
+	{
+		FString Message(e.what());
+		UE_LOG(LogAudioCaptureCore, Error, TEXT("Failed to %s: %s."), Action, *Message);
+	}
+	catch (...)
+	{
+		UE_LOG(LogAudioCaptureCore, Error, TEXT("Failed to %s."), Action);
+	}
+
+	return InvalidDeviceID;
+#else // ^^^ WITH_RTAUDIO vvv !WITH_RTAUDIO
+	return InvalidDeviceID;
+#endif // !WITH_RTAUDIO
+}
+
diff --git a/Engine/Source/Runtime/AudioCaptureImplementations/AudioCaptureRtAudio/Private/AudioCaptureRtAudio.h b/Engine/Source/Runtime/AudioCaptureImplementations/AudioCaptureRtAudio/Private/AudioCaptureRtAudio.h
--- a/Engine/Source/Runtime/AudioCaptureImplementations/AudioCaptureRtAudio/Private/AudioCaptureRtAudio.h
+++ b/Engine/Source/Runtime/AudioCaptureImplementations/AudioCaptureRtAudio/Private/AudioCaptureRtAudio.h
@@ -25,7 +25,7 @@ namespace Audio
 		FAudioCaptureRtAudioStream();
 
 		virtual bool GetCaptureDeviceInfo(FCaptureDeviceInfo& OutInfo, int32 DeviceIndex) override;
-		virtual bool OpenCaptureStream(const FAudioCaptureDeviceParams& InParams, FOnCaptureFunction InOnCapture, uint32 NumFramesDesired) override;
+		virtual bool OpenAudioCaptureStream(const FAudioCaptureDeviceParams& InParams, FOnAudioCaptureFunction InOnCapture, uint32 NumFramesDesired) override;
 		virtual bool CloseStream() override;
 		virtual bool StartStream() override;
 		virtual bool StopStream() override;
@@ -38,10 +38,12 @@ namespace Audio
 		virtual bool GetInputDevicesAvailable(TArray<FCaptureDeviceInfo>& OutDevices) override;
 
 	private:
-		FOnCaptureFunction OnCapture;
+		FOnAudioCaptureFunction OnCapture;
 		int32 NumChannels;
 		int32 SampleRate;
 
+		uint32 GetDefaultInputDevice();
+
 #if WITH_RTAUDIO
 		RtAudio CaptureDevice;
 #endif
diff --git a/Engine/Source/Runtime/AudioCaptureImplementations/AudioCaptureRtAudio/Private/RtAudioWrapper.cpp b/Engine/Source/Runtime/AudioCaptureImplementations/AudioCaptureRtAudio/Private/RtAudioWrapper.cpp
--- a/Engine/Source/Runtime/AudioCaptureImplementations/AudioCaptureRtAudio/Private/RtAudioWrapper.cpp
+++ b/Engine/Source/Runtime/AudioCaptureImplementations/AudioCaptureRtAudio/Private/RtAudioWrapper.cpp
@@ -63,6 +63,8 @@ namespace Audio
 		return FDeviceInfo{};
 	}
 
+	// Suppress deprecation warning for use of FAudioCallback
+	PRAGMA_DISABLE_DEPRECATION_WARNINGS
 	// Opens an audio stream. Returns true on success, false on error. 
 	bool FRtAudioInputWrapper::OpenStream(const FStreamParameters& InStreamParams, float InDesiredSampleRate, uint32* InOutDesiredBufferNumFrames, FAudioCallback Callback, void* InUserData)
 	{
@@ -90,6 +92,7 @@ namespace Audio
 
 		return true;
 	}
+	PRAGMA_ENABLE_DEPRECATION_WARNINGS
 
 	// Starts an open stream
 	void FRtAudioInputWrapper::StartStream()
diff --git a/Engine/Source/Runtime/AudioCaptureImplementations/AudioCaptureRtAudio/Public/RtAudioWrapper.h b/Engine/Source/Runtime/AudioCaptureImplementations/AudioCaptureRtAudio/Public/RtAudioWrapper.h
--- a/Engine/Source/Runtime/AudioCaptureImplementations/AudioCaptureRtAudio/Public/RtAudioWrapper.h
+++ b/Engine/Source/Runtime/AudioCaptureImplementations/AudioCaptureRtAudio/Public/RtAudioWrapper.h
@@ -10,7 +10,7 @@ namespace Audio
 {
 	// FRtAudioInputWrapper wraps RtAudio implementation and provides a way to 
 	// capture input audio from an hardware audio device. 
-	class AUDIOCAPTURERTAUDIO_API FRtAudioInputWrapper
+	class AUDIOCAPTURERTAUDIO_API UE_DEPRECATED(5.3, "FRtAudioInputWrapper is deprecated, please use FAudioCapture instead.") FRtAudioInputWrapper
 	{
 	public:
 		static constexpr uint32 InvalidDeviceID = INDEX_NONE;
diff --git a/Engine/Source/Runtime/AudioCaptureImplementations/IOS/AudioCaptureAudioUnit/Private/AudioCaptureAudioUnit.cpp b/Engine/Source/Runtime/AudioCaptureImplementations/IOS/AudioCaptureAudioUnit/Private/AudioCaptureAudioUnit.cpp
--- a/Engine/Source/Runtime/AudioCaptureImplementations/IOS/AudioCaptureAudioUnit/Private/AudioCaptureAudioUnit.cpp
+++ b/Engine/Source/Runtime/AudioCaptureImplementations/IOS/AudioCaptureAudioUnit/Private/AudioCaptureAudioUnit.cpp
@@ -83,7 +83,7 @@ bool Audio::FAudioCaptureAudioUnitStream::GetCaptureDeviceInfo(FCaptureDeviceInf
 	return true;
 }
 
-bool Audio::FAudioCaptureAudioUnitStream::OpenCaptureStream(const FAudioCaptureDeviceParams& InParams, FOnCaptureFunction InOnCapture, uint32 NumFramesDesired)
+bool Audio::FAudioCaptureAudioUnitStream::OpenAudioCaptureStream(const FAudioCaptureDeviceParams& InParams, FOnAudioCaptureFunction InOnCapture, uint32 NumFramesDesired)
 {
 	NumChannels = 1;
 	SampleRate = 48000;
@@ -219,8 +219,7 @@ bool Audio::FAudioCaptureAudioUnitStream::IsCapturing() const
 
 void Audio::FAudioCaptureAudioUnitStream::OnAudioCapture(void* InBuffer, uint32 InBufferFrames, double StreamTime, bool bOverflow)
 {
-	float* InBufferData = (float*)InBuffer;
-	OnCapture(InBufferData, InBufferFrames, NumChannels, SampleRate, StreamTime, bOverflow);
+	OnCapture(InBuffer, InBufferFrames, NumChannels, SampleRate, StreamTime, bOverflow);
 }
 
 bool Audio::FAudioCaptureAudioUnitStream::GetInputDevicesAvailable(TArray<FCaptureDeviceInfo>& OutDevices)
diff --git a/Engine/Source/Runtime/AudioCaptureImplementations/IOS/AudioCaptureAudioUnit/Private/AudioCaptureAudioUnit.h b/Engine/Source/Runtime/AudioCaptureImplementations/IOS/AudioCaptureAudioUnit/Private/AudioCaptureAudioUnit.h
--- a/Engine/Source/Runtime/AudioCaptureImplementations/IOS/AudioCaptureAudioUnit/Private/AudioCaptureAudioUnit.h
+++ b/Engine/Source/Runtime/AudioCaptureImplementations/IOS/AudioCaptureAudioUnit/Private/AudioCaptureAudioUnit.h
@@ -13,7 +13,7 @@ namespace Audio
 		FAudioCaptureAudioUnitStream();
 
 		virtual bool GetCaptureDeviceInfo(FCaptureDeviceInfo& OutInfo, int32 DeviceIndex) override;
-		virtual bool OpenCaptureStream(const FAudioCaptureDeviceParams& InParams, FOnCaptureFunction InOnCapture, uint32 NumFramesDesired) override;
+		virtual bool OpenAudioCaptureStream(const FAudioCaptureDeviceParams& InParams, FOnAudioCaptureFunction InOnCapture, uint32 NumFramesDesired) override;
 		virtual bool CloseStream() override;
 		virtual bool StartStream() override;
 		virtual bool StopStream() override;
@@ -33,7 +33,7 @@ namespace Audio
 		
 		AudioComponentInstance IOUnit;
 		
-		FOnCaptureFunction OnCapture;
+		FOnAudioCaptureFunction OnCapture;
 		int32 NumChannels;
 		int32 SampleRate;
 		TArray<uint8> CaptureBuffer;
diff --git a/Engine/Source/Runtime/AudioCaptureImplementations/Windows/AudioCaputureWasapi/AudioCaptureWasapi.Build.cs b/Engine/Source/Runtime/AudioCaptureImplementations/Windows/AudioCaputureWasapi/AudioCaptureWasapi.Build.cs
new file mode 100644
--- /dev/null
+++ b/Engine/Source/Runtime/AudioCaptureImplementations/Windows/AudioCaputureWasapi/AudioCaptureWasapi.Build.cs
@@ -0,0 +1,17 @@
+// Copyright Epic Games, Inc. All Rights Reserved.
+
+using UnrealBuildTool;
+
+public class AudioCaptureWasapi : ModuleRules
+{
+	public AudioCaptureWasapi(ReadOnlyTargetRules Target) : base(Target)
+	{
+		PrivateDependencyModuleNames.Add("Core");
+		PrivateDependencyModuleNames.Add("AudioCaptureCore");
+
+		if (Target.Platform.IsInGroup(UnrealPlatformGroup.Windows))
+		{
+			PublicDefinitions.Add("WITH_AUDIOCAPTURE=1");
+		}
+	}
+}
diff --git a/Engine/Source/Runtime/AudioCaptureImplementations/Windows/AudioCaputureWasapi/Private/AudioCaptureWasapi.cpp b/Engine/Source/Runtime/AudioCaptureImplementations/Windows/AudioCaputureWasapi/Private/AudioCaptureWasapi.cpp
new file mode 100644
--- /dev/null
+++ b/Engine/Source/Runtime/AudioCaptureImplementations/Windows/AudioCaputureWasapi/Private/AudioCaptureWasapi.cpp
@@ -0,0 +1,242 @@
+// Copyright Epic Games, Inc. All Rights Reserved.
+
+#include "AudioCaptureWasapi.h"
+
+
+namespace Audio
+{
+	FAudioCaptureWasapiStream::FAudioCaptureWasapiStream()
+	{
+	}
+
+	bool FAudioCaptureWasapiStream::GetCaptureDeviceInfo(FCaptureDeviceInfo& OutInfo, int32 DeviceIndex)
+	{
+		FString DeviceId;
+		if (DeviceIndex == DefaultDeviceIndex)
+		{
+			DeviceId = CaptureDevice.GetDefaultInputDeviceId();
+		}
+		else
+		{
+			CaptureDevice.GetDeviceIdFromIndex(DeviceIndex, eAll, DeviceId);
+		}
+
+		if (!DeviceId.IsEmpty())
+		{
+			FWasapiDeviceEnumeration::FDeviceInfo DeviceInfo;
+
+			// DeviceIndex is for collection of all devices here
+			if (CaptureDevice.GetDeviceInfo(DeviceId, DeviceInfo))
+			{
+				FCaptureDeviceInfo CaptureDeviceInfo;
+				CaptureDeviceInfo.DeviceId = DeviceInfo.DeviceId;
+				CaptureDeviceInfo.DeviceName = DeviceInfo.FriendlyName;
+				CaptureDeviceInfo.InputChannels = DeviceInfo.NumInputChannels;
+				CaptureDeviceInfo.PreferredSampleRate = DeviceInfo.PreferredSampleRate;
+				CaptureDeviceInfo.bSupportsHardwareAEC = false;
+
+				OutInfo = MoveTemp(CaptureDeviceInfo);
+
+				return true;
+			}
+		}
+
+		return false;
+	}
+
+	bool FAudioCaptureWasapiStream::OpenAudioCaptureStream(const FAudioCaptureDeviceParams& InParams, FOnAudioCaptureFunction InOnCapture, uint32 NumFramesDesired)
+	{
+		if (CaptureDevice.IsStreamOpen())
+		{
+			CaptureDevice.StopStream();
+			CaptureDevice.CloseStream();
+		}
+
+		FString DeviceId;
+		if (InParams.DeviceIndex == DefaultDeviceIndex)
+		{
+			DeviceId = CaptureDevice.GetDefaultInputDeviceId();
+		}
+		else
+		{
+			// InParams.DeviceIndex is from the capture subset of devices 
+			CaptureDevice.GetDeviceIdFromIndex(InParams.DeviceIndex, eCapture, DeviceId);
+		}
+
+		if (!DeviceId.IsEmpty())
+		{
+			FWasapiDeviceEnumeration::FDeviceInfo DeviceInfo;
+			if (CaptureDevice.GetDeviceInfo(DeviceId, DeviceInfo))
+			{
+				FWasapiAudioFormat AudioFormat;
+				if (GetAudioFormatFromCaptureParams(InParams, DeviceInfo, AudioFormat))
+				{
+					FWasapiOnAudioCaptureFunction OnAudioCaptureCallback = [this](void* InBuffer, uint32 InNumFrames, double InStreamPosition, bool bInDiscontinuityError)
+					{
+						OnAudioCapture(InBuffer, InNumFrames, InStreamPosition, bInDiscontinuityError);
+					};
+
+					if (CaptureDevice.OpenStream(DeviceId, AudioFormat, NumFramesDesired, MoveTemp(OnAudioCaptureCallback)))
+					{
+						NumChannels = AudioFormat.GetNumChannels();
+						SampleRate = AudioFormat.GetSampleRate();
+						OnCapture = MoveTemp(InOnCapture);
+
+						return true;
+					}
+				}
+			}
+		}
+
+		return false;
+	}
+
+	bool FAudioCaptureWasapiStream::CloseStream()
+	{
+		if (CaptureDevice.IsStreamOpen())
+		{
+			CaptureDevice.CloseStream();
+
+			return !CaptureDevice.IsStreamOpen();
+		}
+
+		return true;
+	}
+
+	bool FAudioCaptureWasapiStream::StartStream()
+	{
+		CaptureDevice.StartStream();
+
+		return CaptureDevice.IsCapturing();
+	}
+
+	bool FAudioCaptureWasapiStream::StopStream()
+	{
+		if (CaptureDevice.IsStreamOpen())
+		{
+			CaptureDevice.StopStream();
+
+			return !CaptureDevice.IsCapturing();
+		}
+
+		return true;
+	}
+
+	bool FAudioCaptureWasapiStream::AbortStream()
+	{
+		if (CaptureDevice.IsStreamOpen())
+		{
+			CaptureDevice.AbortStream();
+
+			return !CaptureDevice.IsStreamOpen();
+		}
+
+		return true;
+	}
+
+	bool FAudioCaptureWasapiStream::GetStreamTime(double& OutStreamTime)
+	{
+		OutStreamTime = CaptureDevice.GetStreamPosition();
+		return true;
+	}
+
+	bool FAudioCaptureWasapiStream::IsStreamOpen() const
+	{
+		return CaptureDevice.IsStreamOpen();
+	}
+
+	bool FAudioCaptureWasapiStream::IsCapturing() const
+	{
+		return CaptureDevice.IsCapturing();
+	}
+
+	void FAudioCaptureWasapiStream::OnAudioCapture(void* InBuffer, uint32 InBufferFrames, double StreamTime, bool bOverflow)
+	{
+		OnCapture(InBuffer, InBufferFrames, NumChannels, SampleRate, StreamTime, bOverflow);
+	}
+
+	bool FAudioCaptureWasapiStream::GetInputDevicesAvailable(TArray<FCaptureDeviceInfo>& OutDevices)
+	{
+		TArray<FWasapiDeviceEnumeration::FDeviceInfo> Devices;
+		if (CaptureDevice.GetInputDevicesAvailable(Devices))
+		{
+			for (const FWasapiDeviceEnumeration::FDeviceInfo& DeviceInfo : Devices)
+			{
+				FCaptureDeviceInfo CaptureDeviceInfo;
+				CaptureDeviceInfo.DeviceId = DeviceInfo.DeviceId;
+				CaptureDeviceInfo.DeviceName = DeviceInfo.FriendlyName;
+				CaptureDeviceInfo.InputChannels = DeviceInfo.NumInputChannels;
+				CaptureDeviceInfo.PreferredSampleRate = DeviceInfo.PreferredSampleRate;
+				CaptureDeviceInfo.bSupportsHardwareAEC = false;
+
+				OutDevices.Emplace(MoveTemp(CaptureDeviceInfo));
+			}
+
+			return true;
+		}
+
+		return false;
+	}
+
+	bool FAudioCaptureWasapiStream::GetAudioFormatFromCaptureParams(
+		const FAudioCaptureDeviceParams& InParams, 
+		const FWasapiDeviceEnumeration::FDeviceInfo& InDeviceInfo,
+		FWasapiAudioFormat& OutAudioFormat)
+	{
+		int32 NumChannels = InParams.NumInputChannels;
+		int32 SampleRate = InParams.SampleRate;
+
+		// If the input params didn't specify num channels, match the device
+		if (NumChannels == InvalidDeviceChannelCount)
+		{
+			NumChannels = InDeviceInfo.NumInputChannels;
+		}
+
+		// If the input params didn't specify sample rate, match the device
+		if (SampleRate == InvalidDeviceSampleRate)
+		{
+			SampleRate = InDeviceInfo.PreferredSampleRate;
+		}
+
+		// If the input params didn't specify a bit depth, default to 32-bit float
+		EWasapiAudioEncoding AudioEncoding;
+		switch (InParams.PCMAudioEncoding)
+		{
+		case EPCMAudioEncoding::PCM_8:
+			AudioEncoding = EWasapiAudioEncoding::PCM_8;
+			break;
+
+		case EPCMAudioEncoding::PCM_16:
+			AudioEncoding = EWasapiAudioEncoding::PCM_16;
+			break;
+
+		case EPCMAudioEncoding::PCM_24:
+			AudioEncoding = EWasapiAudioEncoding::PCM_24;
+			break;
+
+		case EPCMAudioEncoding::PCM_24_IN_32:
+			AudioEncoding = EWasapiAudioEncoding::PCM_24_IN_32;
+			break;
+
+		case EPCMAudioEncoding::PCM_32:
+			AudioEncoding = EWasapiAudioEncoding::PCM_32;
+			break;
+
+		case EPCMAudioEncoding::FLOATING_POINT_32:
+			AudioEncoding = EWasapiAudioEncoding::FLOATING_POINT_32;
+			break;
+
+		case EPCMAudioEncoding::FLOATING_POINT_64:
+			AudioEncoding = EWasapiAudioEncoding::FLOATING_POINT_64;
+			break;
+
+		case EPCMAudioEncoding::UNKNOWN:
+		default:
+			return false;
+		}
+
+		OutAudioFormat = FWasapiAudioFormat(NumChannels, SampleRate, AudioEncoding);
+
+		return true;
+	}
+}
diff --git a/Engine/Source/Runtime/AudioCaptureImplementations/Windows/AudioCaputureWasapi/Private/AudioCaptureWasapi.h b/Engine/Source/Runtime/AudioCaptureImplementations/Windows/AudioCaputureWasapi/Private/AudioCaptureWasapi.h
new file mode 100644
--- /dev/null
+++ b/Engine/Source/Runtime/AudioCaptureImplementations/Windows/AudioCaputureWasapi/Private/AudioCaptureWasapi.h
@@ -0,0 +1,70 @@
+// Copyright Epic Games, Inc. All Rights Reserved.
+
+#pragma once
+
+#include "AudioCaptureCore.h"
+#include "AudioCaptureDeviceInterface.h"
+#include "WasapiStreamManager.h"
+
+
+namespace Audio
+{
+	/**
+	 * WASAPI implementation of IAudioCaptureStream which encapsulates
+	 * capturing audio streams from device.
+	 */
+	class FAudioCaptureWasapiStream : public IAudioCaptureStream
+	{
+	public:
+		FAudioCaptureWasapiStream();
+
+		// Begin IAudioCaptureStream overrides
+		virtual bool GetCaptureDeviceInfo(FCaptureDeviceInfo& OutInfo, int32 DeviceIndex) override;
+
+		virtual bool OpenAudioCaptureStream(const FAudioCaptureDeviceParams& InParams, FOnAudioCaptureFunction InOnCapture, uint32 NumFramesDesired) override;
+		
+		virtual bool CloseStream() override;
+
+		virtual bool StartStream() override;
+
+		virtual bool StopStream() override;
+
+		virtual bool AbortStream() override;
+
+		virtual bool GetStreamTime(double& OutStreamTime) override;
+
+		virtual int32 GetSampleRate() const override { return SampleRate; }
+
+		virtual bool IsStreamOpen() const override;
+
+		virtual bool IsCapturing() const override;
+
+		virtual void OnAudioCapture(void* InBuffer, uint32 InBufferFrames, double StreamTime, bool bOverflow) override;
+
+		virtual bool GetInputDevicesAvailable(TArray<FCaptureDeviceInfo>& OutDevices) override;
+		// End IAudioCaptureStream overrides
+
+	private:
+		/** Audio callback which will be called during audio capture. */
+		FOnAudioCaptureFunction OnCapture;
+		/** The number of channels to capture audio data on from the device. */
+		int32 NumChannels = InvalidDeviceChannelCount;
+		/** The sample rate used during this capture. */
+		int32 SampleRate = InvalidDeviceSampleRate;
+		/** The aggregate capture device which manages one or more capture streams. */
+		FWasapiStreamManager CaptureDevice;
+
+		/**
+		 * GetAudioFormatFromCaptureParams - Converts FAudioCaptureDeviceParams into a corresponding FWasapiAudioFormat
+		 * utilizing InDeviceInfo for default values if needed.
+		 * 
+		 * @param InParams - The capture parameters supplied by the caller to be used when configuring the audio format.
+		 * @param InDeviceInfo - The audio device info which will be used as a fallback for default format values if needed.
+		 * @param OutAudioFormat - The resulting audio format to corresponding to the input parameters.
+		 */
+		static bool GetAudioFormatFromCaptureParams(
+			const FAudioCaptureDeviceParams& InParams, 
+			const FWasapiDeviceEnumeration::FDeviceInfo& InDeviceInfo,
+			FWasapiAudioFormat& OutAudioFormat);
+	};
+}
diff --git a/Engine/Source/Runtime/AudioCaptureImplementations/Windows/AudioCaputureWasapi/Private/AudioCaptureWasapiModule.cpp b/Engine/Source/Runtime/AudioCaptureImplementations/Windows/AudioCaputureWasapi/Private/AudioCaptureWasapiModule.cpp
new file mode 100644
--- /dev/null
+++ b/Engine/Source/Runtime/AudioCaptureImplementations/Windows/AudioCaputureWasapi/Private/AudioCaptureWasapiModule.cpp
@@ -0,0 +1,38 @@
+// Copyright Epic Games, Inc. All Rights Reserved.
+
+#include "AudioCaptureDeviceInterface.h"
+#include "AudioCaptureWasapi.h"
+#include "Features/IModularFeatures.h"
+#include "Modules/ModuleManager.h"
+
+namespace Audio
+{
+	
+	class FAudioCaptureWasapiFactory : public IAudioCaptureFactory
+	{
+	public:
+		virtual TUniquePtr<IAudioCaptureStream> CreateNewAudioCaptureStream() override
+		{
+			return TUniquePtr<IAudioCaptureStream>(new FAudioCaptureWasapiStream());
+		}
+	};
+
+	class FAudioCaptureWasapiModule : public IModuleInterface
+	{
+	private:
+		FAudioCaptureWasapiFactory AudioCaptureFactory;
+
+	public:
+		virtual void StartupModule() override
+		{
+			IModularFeatures::Get().RegisterModularFeature(IAudioCaptureFactory::GetModularFeatureName(), &AudioCaptureFactory);
+		}
+
+		virtual void ShutdownModule() override
+		{
+			IModularFeatures::Get().UnregisterModularFeature(IAudioCaptureFactory::GetModularFeatureName(), &AudioCaptureFactory);
+		}
+	};
+}
+
+IMPLEMENT_MODULE(Audio::FAudioCaptureWasapiModule, AudioCaptureWasapi)
diff --git a/Engine/Source/Runtime/AudioCaptureImplementations/Windows/AudioCaputureWasapi/Private/WasapiAudioFormat.cpp b/Engine/Source/Runtime/AudioCaptureImplementations/Windows/AudioCaputureWasapi/Private/WasapiAudioFormat.cpp
new file mode 100644
--- /dev/null
+++ b/Engine/Source/Runtime/AudioCaptureImplementations/Windows/AudioCaputureWasapi/Private/WasapiAudioFormat.cpp
@@ -0,0 +1,277 @@
+// Copyright Epic Games, Inc. All Rights Reserved.
+
+#include "WasapiAudioFormat.h"
+
+#include "AudioCaptureCoreLog.h"
+
+
+namespace Audio
+{
+	// For explanation on WAVEFORMATEX structure and cbSize requirements, see: 
+	// https://learn.microsoft.com/en-us/windows/win32/api/mmeapi/ns-mmeapi-waveformatex
+	static constexpr uint16 WaveFormatExtensibleCbsize = 22;
+	static constexpr uint16 DefaultChannelCount = 2;
+	static constexpr uint32 DefaultSampleRate = 44100;
+	static constexpr EWasapiAudioEncoding DefaultEncoding = EWasapiAudioEncoding::PCM_16;
+
+
+	FWasapiAudioFormat::FWasapiAudioFormat()
+	{
+		InitAudioFormat(DefaultChannelCount, DefaultSampleRate, DefaultEncoding);
+		InitAudioEncoding();
+	}
+
+	FWasapiAudioFormat::FWasapiAudioFormat(WAVEFORMATEX* InFormat)
+	{
+		if ((InFormat->wFormatTag == WAVE_FORMAT_EXTENSIBLE) && (InFormat->cbSize >= WaveFormatExtensibleCbsize))
+		{
+			FMemory::Memcpy(&WaveFormat, InFormat, sizeof(WAVEFORMATEXTENSIBLE));
+		}
+		else
+		{
+			FMemory::Memcpy(&WaveFormat, InFormat, sizeof(WAVEFORMATEX));
+		}
+
+		InitAudioEncoding();
+	}
+
+	FWasapiAudioFormat::FWasapiAudioFormat(uint16 InChannels, uint32 InSampleRate, EWasapiAudioEncoding InEncoding)
+	{
+		InitAudioFormat(InChannels, InSampleRate, InEncoding);
+		InitAudioEncoding();
+	}
+
+	void FWasapiAudioFormat::InitAudioFormat(uint16 InChannels, uint32 InSampleRate, EWasapiAudioEncoding InEncoding)
+	{
+		static constexpr uint32 BitsPerByte = 8;
+		uint16 BitDepth = EncodingToBitDepth(InEncoding);
+		ensure(BitDepth != 0);
+
+		if (BitDepth != 0)
+		{
+			WaveFormat.Format.wFormatTag = WAVE_FORMAT_EXTENSIBLE;
+			WaveFormat.Format.cbSize = WaveFormatExtensibleCbsize;
+			WaveFormat.Format.nChannels = InChannels;
+			WaveFormat.Format.nSamplesPerSec = InSampleRate;
+			WaveFormat.Format.wBitsPerSample = BitDepth;
+			WaveFormat.Format.nBlockAlign = InChannels * (BitDepth / BitsPerByte);
+			WaveFormat.Format.nAvgBytesPerSec = InSampleRate * WaveFormat.Format.nBlockAlign;
+			WaveFormat.Samples.wValidBitsPerSample = BitDepth;
+
+			if (InEncoding == EWasapiAudioEncoding::FLOATING_POINT_32 || InEncoding == EWasapiAudioEncoding::FLOATING_POINT_64)
+			{
+				WaveFormat.SubFormat = KSDATAFORMAT_SUBTYPE_IEEE_FLOAT;
+			}
+			else
+			{
+				WaveFormat.SubFormat = KSDATAFORMAT_SUBTYPE_PCM;
+				if (InEncoding == EWasapiAudioEncoding::PCM_24_IN_32)
+				{
+					WaveFormat.Samples.wValidBitsPerSample = 24;
+				}
+			}
+		}
+	}
+
+	void FWasapiAudioFormat::InitAudioEncoding()
+	{
+		Encoding = DetermineAudioEncoding(WaveFormat);
+		if (Encoding == EWasapiAudioEncoding::UNKNOWN)
+		{
+			UE_LOG(LogAudioCaptureCore, Error, TEXT("FWasapiAudioFormat unknown audio format"));
+		}
+	}
+
+	EWasapiAudioEncoding FWasapiAudioFormat::GetEncoding() const
+	{
+		return Encoding;
+	}
+
+	uint32 FWasapiAudioFormat::GetNumChannels() const
+	{
+		return WaveFormat.Format.nChannels;
+	}
+
+	uint32 FWasapiAudioFormat::GetSampleRate() const
+	{
+		return WaveFormat.Format.nSamplesPerSec;
+	}
+
+	uint32 FWasapiAudioFormat::GetBitsPerSample() const
+	{
+		return WaveFormat.Format.wBitsPerSample;
+	}
+
+	uint32 FWasapiAudioFormat::GetBytesPerSample() const
+	{
+		return WaveFormat.Format.wBitsPerSample / 8;
+	}
+
+	uint32 FWasapiAudioFormat::GetFrameSizeInBytes() const
+	{
+		return WaveFormat.Format.nBlockAlign;
+	}
+
+	const WAVEFORMATEX* FWasapiAudioFormat::GetWaveFormat() const
+	{
+		return &(WaveFormat.Format);
+	}
+
+	EWasapiAudioEncoding FWasapiAudioFormat::DetermineAudioEncoding(const WAVEFORMATEXTENSIBLE& InFormat)
+	{
+		if (InFormat.Format.wFormatTag == WAVE_FORMAT_IEEE_FLOAT)
+		{
+			if (InFormat.Format.wBitsPerSample == (sizeof(float) * 8))
+			{
+				return EWasapiAudioEncoding::FLOATING_POINT_32;
+			}
+			else if (InFormat.Format.wBitsPerSample == (sizeof(double) * 8))
+			{
+				return EWasapiAudioEncoding::FLOATING_POINT_64;
+			}
+			else
+			{
+				return EWasapiAudioEncoding::UNKNOWN;
+			}
+		}
+		else
+		{
+			if (InFormat.Format.wFormatTag == WAVE_FORMAT_PCM)
+			{
+				switch (InFormat.Format.wBitsPerSample)
+				{
+				case 8:
+					return EWasapiAudioEncoding::PCM_8;
+
+				case 16:
+					return EWasapiAudioEncoding::PCM_16;
+
+				case 24:
+					return EWasapiAudioEncoding::PCM_24;
+
+				case 32:
+					return EWasapiAudioEncoding::PCM_32;
+
+				default:
+					return EWasapiAudioEncoding::UNKNOWN;
+				}
+			}
+			else
+			{
+				if ((InFormat.Format.wFormatTag == WAVE_FORMAT_EXTENSIBLE) && (InFormat.Format.cbSize >= WaveFormatExtensibleCbsize))
+				{
+					if (InFormat.SubFormat == KSDATAFORMAT_SUBTYPE_IEEE_FLOAT)
+					{
+						if (InFormat.Format.wBitsPerSample == (sizeof(float) * 8))
+						{
+							return EWasapiAudioEncoding::FLOATING_POINT_32;
+						}
+						else if (InFormat.Format.wBitsPerSample == (sizeof(double) * 8))
+						{
+							return EWasapiAudioEncoding::FLOATING_POINT_64;
+						}
+						else
+						{
+							return EWasapiAudioEncoding::UNKNOWN;
+						}
+					}
+					else
+					{
+						if (InFormat.SubFormat == KSDATAFORMAT_SUBTYPE_PCM)
+						{
+							switch (InFormat.Samples.wValidBitsPerSample)
+							{
+							case 8:
+								return (InFormat.Format.wBitsPerSample == 8) ? EWasapiAudioEncoding::PCM_8 : EWasapiAudioEncoding::UNKNOWN;
+
+							case 16:
+								return (InFormat.Format.wBitsPerSample == 16) ? EWasapiAudioEncoding::PCM_16 : EWasapiAudioEncoding::UNKNOWN;
+
+							case 32:
+								return (InFormat.Format.wBitsPerSample == 32) ? EWasapiAudioEncoding::PCM_32 : EWasapiAudioEncoding::UNKNOWN;
+
+							case 24:
+								if (InFormat.Format.wBitsPerSample == 32)
+								{
+									return EWasapiAudioEncoding::PCM_24_IN_32;
+								}
+								else
+								{
+									if (InFormat.Format.wBitsPerSample == 24)
+									{
+										return EWasapiAudioEncoding::PCM_24;
+									}
+									else
+									{
+										return EWasapiAudioEncoding::UNKNOWN;
+									}
+								}
+							default:
+								break;
+							}
+						}
+					}
+				}
+			}
+		}
+
+		return EWasapiAudioEncoding::UNKNOWN;
+	}
+
+	uint16 FWasapiAudioFormat::EncodingToBitDepth(EWasapiAudioEncoding InEncoding)
+	{
+		switch (InEncoding)
+		{
+		case EWasapiAudioEncoding::PCM_8:
+			return 8;
+
+		case EWasapiAudioEncoding::PCM_16:
+			return 16;
+
+		case EWasapiAudioEncoding::PCM_24:
+			return 24;
+
+		case EWasapiAudioEncoding::PCM_24_IN_32:
+		case EWasapiAudioEncoding::PCM_32:
+		case EWasapiAudioEncoding::FLOATING_POINT_32:
+			return 32;
+
+		case EWasapiAudioEncoding::FLOATING_POINT_64:
+			return 64;
+
+		default:
+			return 0;
+		}
+	}
+
+	FString FWasapiAudioFormat::GetEncodingString() const
+	{
+		switch (GetEncoding())
+		{
+		case EWasapiAudioEncoding::PCM_8:
+			return TEXT("8-Bit PCM");
+
+		case EWasapiAudioEncoding::PCM_16:
+			return TEXT("16-Bit PCM");
+
+		case EWasapiAudioEncoding::PCM_24:
+			return TEXT("24-Bit PCM");
+
+		case EWasapiAudioEncoding::PCM_24_IN_32:
+			return TEXT("24-Bit PCM (32-bit Container)");
+
+		case EWasapiAudioEncoding::PCM_32:
+			return TEXT("PCM_32 PCM");
+
+		case EWasapiAudioEncoding::FLOATING_POINT_32:
+			return TEXT("32-Bit Floating Point");
+
+		case EWasapiAudioEncoding::FLOATING_POINT_64:
+			return TEXT("64-Bit Floating Point");
+
+		default:
+			return TEXT("UNKNOWN");
+		}
+	}
+
+}
diff --git a/Engine/Source/Runtime/AudioCaptureImplementations/Windows/AudioCaputureWasapi/Private/WasapiAudioFormat.h b/Engine/Source/Runtime/AudioCaptureImplementations/Windows/AudioCaputureWasapi/Private/WasapiAudioFormat.h
new file mode 100644
--- /dev/null
+++ b/Engine/Source/Runtime/AudioCaptureImplementations/Windows/AudioCaputureWasapi/Private/WasapiAudioFormat.h
@@ -0,0 +1,104 @@
+// Copyright Epic Games, Inc. All Rights Reserved.
+
+#pragma once
+
+#include "AudioCaptureCore.h"
+
+#include "Windows/AllowWindowsPlatformTypes.h"
+#include "Windows/AllowWindowsPlatformAtomics.h"
+
+THIRD_PARTY_INCLUDES_START
+#include <AudioClient.h>
+THIRD_PARTY_INCLUDES_END
+
+#include "Windows/HideWindowsPlatformAtomics.h"
+#include "Windows/HideWindowsPlatformTypes.h"
+
+
+namespace Audio
+{
+	/**
+	 * EWasapiAudioEncoding - Audio bit depths supported by WASAPI implementation
+	 */
+	enum class EWasapiAudioEncoding : uint8
+	{
+		UNKNOWN,
+		PCM_8,
+		PCM_16,
+		PCM_24,
+		PCM_24_IN_32,
+		PCM_32,
+		FLOATING_POINT_32,
+		FLOATING_POINT_64
+	};
+
+	/**
+	 * FWasapiAudioFormat - Wrapper class for Windows WAVEFORMATEX and WAVEFORMATEXTENSIBLE structs.
+	 */
+	class FWasapiAudioFormat
+	{
+	public:
+		/** FWasapiAudioFormat - Constructs underlying WAVEFORMATEX with default values. */
+		FWasapiAudioFormat();
+
+		/**
+		 * FWasapiAudioFormat - Constructs underlying data structure with given WAVEFORMATEX parameter.
+		 * 
+		 * @param InFormat - Pointer to a WAVEFORMATEX struct from which to construct this object. 
+		 * Correctly handles the extended form, WAVEFORMATEXTENSIBLE. 
+		 */
+		explicit FWasapiAudioFormat(WAVEFORMATEX* InFormat);
+
+		/**
+		 * FWasapiAudioFormat - Constructs underlying data structure using given parameters.
+		 * 
+		 * @param InChannels - Number of audio channels.
+		 * @param InSampleRate - Samples per second (Hz).
+		 * @param InEncoding - For PCM data, this is the bit depth of data (e.g. 16-bit integer, 32-bit float, etc.)
+		 */
+		FWasapiAudioFormat(uint16 InChannels, uint32 InSampleRate, EWasapiAudioEncoding InEncoding);
+
+		/** GetEncoding - Returns the encoding for this audio format. */
+		EWasapiAudioEncoding GetEncoding() const;
+
+		/** GetNumChannels - Returns the number of audio channels present in this format. */
+		uint32 GetNumChannels() const;
+
+		/** GetSampleRate - Returns the audio sample rate for this format. */
+		uint32 GetSampleRate() const;
+
+		/** GetBitsPerSample - Returns the bit depth for the audio data as an integer value. */
+		uint32 GetBitsPerSample() const;
+
+		/** GetBytesPerSample - Returns the number of bytes needed to store a single sample of audio data. */
+		uint32 GetBytesPerSample() const;
+
+		/** GetFrameSizeInBytes - Returns the number of bytes needed to store a singe frame of audio data. */
+		uint32 GetFrameSizeInBytes() const;
+		
+		/** GetEncodingString - Returns string form of encoding. Suitable for logging. */
+		FString GetEncodingString() const;
+
+		/**
+		 * GetWaveFormat - Returns pointer to raw WAVEFORMATEX struct. Use this for specific OS API calls
+		 * which require it. Otherwise use the accessors above.
+		 */
+		const WAVEFORMATEX* GetWaveFormat() const;
+
+	private:
+
+		/** InitAudioFormat - Init method used by constructors to initialize data. */
+		void InitAudioFormat(uint16 InChannels, uint32 InSampleRate, EWasapiAudioEncoding InEncoding);
+		/** InitAudioEncoding - Parses underlying WAVEFORMATEXTENSIBLE structure and determines corresponding encoding. */
+		void InitAudioEncoding();
+		/** DetermineAudioEncoding - Helper function for determining audio encoding from WAVEFORMATEXTENSIBLE struct. */
+		static EWasapiAudioEncoding DetermineAudioEncoding(const WAVEFORMATEXTENSIBLE& InFormat);
+		/** EncodingToBitDepth - Helper function for getting bit depth from InEncoding parameter. */
+		static uint16 EncodingToBitDepth(EWasapiAudioEncoding InEncoding);
+
+		/** WaveFormat - Windows wave format struct which this object wraps. */
+		WAVEFORMATEXTENSIBLE WaveFormat = {};
+		/** Encoding - Audio encoding for this object. */
+		EWasapiAudioEncoding Encoding = EWasapiAudioEncoding::UNKNOWN;
+	};
+}
diff --git a/Engine/Source/Runtime/AudioCaptureImplementations/Windows/AudioCaputureWasapi/Private/WasapiCaptureLog.h b/Engine/Source/Runtime/AudioCaptureImplementations/Windows/AudioCaputureWasapi/Private/WasapiCaptureLog.h
new file mode 100644
--- /dev/null
+++ b/Engine/Source/Runtime/AudioCaptureImplementations/Windows/AudioCaputureWasapi/Private/WasapiCaptureLog.h
@@ -0,0 +1,9 @@
+#pragma once
+
+#include "AudioCaptureCoreLog.h"
+
+#define WASAPI_CAPTURE_LOG_RESULT(FunctionName, Result) \
+	{ \
+		FString ErrorString = FString::Printf(TEXT("%s -> 0x%X (line: %d)"), TEXT( FunctionName ), Result, __LINE__); \
+		UE_LOG(LogAudioCaptureCore, Error, TEXT("WasapiCapture Error: %s"), *ErrorString);				  \
+	}
diff --git a/Engine/Source/Runtime/AudioCaptureImplementations/Windows/AudioCaputureWasapi/Private/WasapiCaptureThread.cpp b/Engine/Source/Runtime/AudioCaptureImplementations/Windows/AudioCaputureWasapi/Private/WasapiCaptureThread.cpp
new file mode 100644
--- /dev/null
+++ b/Engine/Source/Runtime/AudioCaptureImplementations/Windows/AudioCaputureWasapi/Private/WasapiCaptureThread.cpp
@@ -0,0 +1,96 @@
+// Copyright Epic Games, Inc. All Rights Reserved.
+
+#include "WasapiCaptureThread.h"
+#include "WasapiCaptureLog.h"
+
+
+namespace Audio
+{
+	FWasapiCaptureRunnable::FWasapiCaptureRunnable(TSharedPtr<FWasapiInputStream> InStream) :
+		InputStream(InStream)
+	{
+	}
+
+	uint32 FWasapiCaptureRunnable::Run()
+	{
+		bIsRunning = true;
+
+		bool bCoInitialized = FWindowsPlatformMisc::CoInitialize(ECOMModel::Multithreaded);
+
+		InputStream->StartStream();
+		while (bIsRunning.load())
+		{
+			if (!InputStream->WaitOnBuffer())
+			{
+				// Accumulate timeouts and report after capture
+				++InputStreamTimeoutsDetected;
+			}
+
+			if (!InputStream->CaptureAudioFrames())
+			{
+				// Accumulate device errors which can happen if the device
+				// goes off line during capture (e.g. unplugged USB device)
+				++InputStreamDeviceErrorsDetected;
+			}
+		}
+
+		InputStream->StopStream();
+
+		if (bCoInitialized)
+		{
+			FWindowsPlatformMisc::CoUninitialize();
+		}
+
+		return 0;
+	}
+
+	void FWasapiCaptureRunnable::Stop()
+	{
+		bIsRunning = false;
+		if (InputStreamTimeoutsDetected > 0)
+		{
+			UE_LOG(LogAudioCaptureCore, Error, TEXT("WasapiCapture Error: InputStream->WaitOnBuffer() reported %d timeouts during capture"), InputStreamTimeoutsDetected);
+		}
+		if (InputStreamDeviceErrorsDetected > 0)
+		{
+			UE_LOG(LogAudioCaptureCore, Error, TEXT("WasapiCapture Error: InputStream->WaitOnBuffer() reported %d device during capture"), InputStreamDeviceErrorsDetected);
+		}
+	}
+
+	FWasapiCaptureThread::FWasapiCaptureThread(TSharedPtr<FWasapiInputStream> InStream) :
+		CaptureRunnable(MakeUnique<FWasapiCaptureRunnable>(InStream))
+	{
+	}
+
+	bool FWasapiCaptureThread::Start()
+	{
+		check(CaptureThread == nullptr);
+
+		CaptureThread = TUniquePtr<FRunnableThread>(FRunnableThread::Create(CaptureRunnable.Get(), TEXT("Audio Capture Thread"), 0, TPri_TimeCritical));
+		return CaptureThread.IsValid();
+	}
+
+	void FWasapiCaptureThread::Stop()
+	{
+		if (CaptureThread.IsValid())
+		{
+			TRACE_CPUPROFILER_EVENT_SCOPE(Audio::FWasapiCaptureThread::Stop);
+
+			bool bShouldWait = true;
+			CaptureThread->Kill(bShouldWait);
+		}
+	}
+
+	void FWasapiCaptureThread::Abort()
+	{
+		if (CaptureThread.IsValid())
+		{
+			TRACE_CPUPROFILER_EVENT_SCOPE(Audio::FWasapiCaptureThread::Abort);
+
+			// Always wait for thread to complete otherwise we can crash if
+			// the stream is disposed of mid-callback.
+			bool bShouldWait = true;
+			CaptureThread->Kill(bShouldWait);
+		}
+	}
+}
diff --git a/Engine/Source/Runtime/AudioCaptureImplementations/Windows/AudioCaputureWasapi/Private/WasapiCaptureThread.h b/Engine/Source/Runtime/AudioCaptureImplementations/Windows/AudioCaputureWasapi/Private/WasapiCaptureThread.h
new file mode 100644
--- /dev/null
+++ b/Engine/Source/Runtime/AudioCaptureImplementations/Windows/AudioCaputureWasapi/Private/WasapiCaptureThread.h
@@ -0,0 +1,96 @@
+// Copyright Epic Games, Inc. All Rights Reserved.
+
+#pragma once
+
+#include "HAL/Runnable.h"
+#include "HAL/RunnableThread.h"
+#include "WasapiInputStream.h"
+
+
+namespace Audio
+{
+	/**
+	 * FWasapiCaptureRunnable - The runnable which executes the main thread loop for FWasapiCaptureThread.
+	 * It has a pointer to a FWasapiInputStream which it uses to coordinating waiting for new data
+	 * and delivery of that data to the audio callback.
+	 */
+	class FWasapiCaptureRunnable : public FRunnable
+	{
+	public:
+		FWasapiCaptureRunnable() = delete;
+		FWasapiCaptureRunnable(FWasapiCaptureRunnable&& InOther) = delete;
+		FWasapiCaptureRunnable(const FWasapiCaptureRunnable& InOther) = delete;
+
+		FWasapiCaptureRunnable& operator=(FWasapiCaptureRunnable&& InOther) = delete;
+		FWasapiCaptureRunnable& operator=(const FWasapiCaptureRunnable& InOther) = delete;
+
+		/**
+		 * FWasapiCaptureRunnable - Constructor with given FWasapiInputStream.
+		 * 
+		 * @param InStream - The stream which will receive audio data as it is received.
+		 */
+		explicit FWasapiCaptureRunnable(TSharedPtr<FWasapiInputStream> InStream);
+		
+		/** Default destructor */
+		virtual ~FWasapiCaptureRunnable() = default;
+
+		// Begin FRunnable overrides
+		virtual uint32 Run() override;
+		virtual void Stop() override;
+		// End FRunnable overrides
+
+	private:
+		/** bIsRunning - The main run loop for this runnable will continue iterating while this flag is true. */
+		std::atomic<bool> bIsRunning;
+		/**
+		 * InputStreamTimeoutsDetected - Accumulates timeouts which occur when the thread event timeout is reached
+		 * prior to the event being signaled for new data being available.
+		 */
+		uint32 InputStreamTimeoutsDetected = 0;
+		/**
+		 * InputStreamDeviceErrorsDetected - Accumulates any device errors which occur. Device errors which can 
+		 * happen if the device goes off line during capture (e.g. unplugged USB device).
+		 */
+		uint32 InputStreamDeviceErrorsDetected = 0;
+		/** InputStream - The input stream which will be run and periodically receive new audio data. */
+		TSharedPtr<FWasapiInputStream> InputStream;
+	};
+
+	/**
+	 * FWasapiCaptureThread - Manages both the FWasapiCaptureRunnable object and the thread whose context it runs in. 
+	 */
+	class FWasapiCaptureThread
+	{
+	public:
+		FWasapiCaptureThread() = delete;
+
+		/**
+		 * FWasapiCaptureThread - Constructor which accepts FWasapiInputStream and creates FWasapiCaptureRunnable with it.
+		 */
+		explicit FWasapiCaptureThread(TSharedPtr<FWasapiInputStream> InStream);
+
+		/**
+		 * Start - Creates the FRunnableThread object which immediately begins running the FWasapiCaptureRunnable member.
+		 * 
+		 * @return - Boolean indicating of the thread was succesfully created.
+		 */
+		bool Start();
+
+		/**
+		 * Stop - Gracefully shuts down the capture thread.
+		 */
+		void Stop();
+
+		/**
+		 * Abort - Performs non-graceful shutdown of capture thread which will close the underyling thread handle 
+		 * without waiting for it to finish.
+		 */
+		void Abort();
+
+	private:
+		/** The thread which is the context that the capture runnable executes in. */
+		TUniquePtr<FRunnableThread> CaptureThread;
+		/** The capture runnable which manages the run loop for the capture stream. */
+		TUniquePtr<FWasapiCaptureRunnable> CaptureRunnable;
+	};
+}
diff --git a/Engine/Source/Runtime/AudioCaptureImplementations/Windows/AudioCaputureWasapi/Private/WasapiDeviceEnumeration.cpp b/Engine/Source/Runtime/AudioCaptureImplementations/Windows/AudioCaputureWasapi/Private/WasapiDeviceEnumeration.cpp
new file mode 100644
--- /dev/null
+++ b/Engine/Source/Runtime/AudioCaptureImplementations/Windows/AudioCaputureWasapi/Private/WasapiDeviceEnumeration.cpp
@@ -0,0 +1,365 @@
+// Copyright Epic Games, Inc. All Rights Reserved.
+
+#include "WasapiDeviceEnumeration.h"
+
+#include "Misc/ScopeExit.h"
+#include "WasapiCaptureLog.h"
+
+#include "Windows/AllowWindowsPlatformTypes.h"
+#include "Windows/AllowWindowsPlatformAtomics.h"
+
+THIRD_PARTY_INCLUDES_START
+#include <Functiondiscoverykeys_devpkey.h>
+THIRD_PARTY_INCLUDES_END
+
+#include "Windows/HideWindowsPlatformAtomics.h"
+#include "Windows/HideWindowsPlatformTypes.h"
+
+
+namespace Audio
+{
+	using FDeviceInfo = FWasapiDeviceEnumeration::FDeviceInfo;
+
+	void FWasapiDeviceEnumeration::Initialize()
+	{
+		// Assumes FWindowsPlatformMisc::CoInitialize() has been called upstream
+		ensure(
+			SUCCEEDED(
+				::CoCreateInstance(__uuidof(MMDeviceEnumerator), nullptr, CLSCTX_INPROC_SERVER, IID_PPV_ARGS(&DeviceEnumerator)) &&
+				DeviceEnumerator != nullptr
+			)
+		);
+
+		EnumerateDefaults();
+	}
+
+	void FWasapiDeviceEnumeration::EnumerateDefaults()
+	{
+		auto GetDefaultDeviceID = [this](EDataFlow InDataFlow, ERole InRole, FString& OutDeviceId) -> bool
+		{
+			// Mark default device.
+			bool bSuccess = false;
+			TComPtr<IMMDevice> DefaultDevice;
+			if (SUCCEEDED(DeviceEnumerator->GetDefaultAudioEndpoint(InDataFlow, InRole, &DefaultDevice)))
+			{
+				LPWSTR IdString = nullptr;
+				if (SUCCEEDED(DefaultDevice->GetId(&IdString)) && IdString)
+				{
+					OutDeviceId = FString(IdString);
+					::CoTaskMemFree(IdString);
+					bSuccess = true;
+				}
+			}
+			return bSuccess;
+		};
+
+		// Get defaults (render, capture).
+		FString DeviceIdName;
+		if (GetDefaultDeviceID(eRender, eConsole, DeviceIdName))
+		{
+			UE_LOG(LogAudioCaptureCore, Verbose, TEXT("WasapiAbstractionLayer: Default Device='%s'"), *DeviceIdName);
+			DefaultRenderId = DeviceIdName;
+		}
+		else
+		{
+			UE_LOG(LogAudioCaptureCore, Error, TEXT("WasapiAbstractionLayer: error fetching default render device"));
+		}
+
+		if (GetDefaultDeviceID(eCapture, eConsole, DeviceIdName))
+		{
+			UE_LOG(LogAudioCaptureCore, Verbose, TEXT("WasapiAbstractionLayer: Default Device='%s'"), *DeviceIdName);
+			DefaultCaptureId = DeviceIdName;
+		}
+		else
+		{
+			UE_LOG(LogAudioCaptureCore, Error, TEXT("WasapiAbstractionLayer: error fetching default capture device"));
+		}
+	}
+
+	FString FWasapiDeviceEnumeration::GetDefaultInputDeviceId()
+	{
+		return DefaultCaptureId;
+	}
+
+	FString FWasapiDeviceEnumeration::GetDefaultOutputDeviceId()
+	{
+		return DefaultRenderId;
+	}
+
+	bool FWasapiDeviceEnumeration::GetDeviceInfo(const FString& InDeviceId, FDeviceInfo& OutDeviceInfo)
+	{
+		TComPtr<IMMDevice> Device;
+		if (GetIMMDevice(InDeviceId, Device))
+		{
+			FDeviceInfo DeviceInfo;
+			if (GetDeviceProperties(Device, DeviceInfo))
+			{
+				OutDeviceInfo = MoveTemp(DeviceInfo);
+				return true;
+			}
+		}
+
+		return false;
+	}
+
+	bool FWasapiDeviceEnumeration::GetDeviceIdFromIndex(int32 InDeviceIndex, EDataFlow InDataFlow, FString& OutDeviceId)
+	{
+		TComPtr<IMMDeviceCollection> Collection;
+		HRESULT Result = DeviceEnumerator->EnumAudioEndpoints(InDataFlow, DEVICE_STATE_ACTIVE, &Collection);
+		if (FAILED(Result))
+		{
+			WASAPI_CAPTURE_LOG_RESULT("IMMDeviceCollection::EnumAudioEndpoints", Result);
+			return false;
+		}
+
+		TComPtr<IMMDevice> Device;
+		Result = Collection->Item(InDeviceIndex, &Device);
+		if (SUCCEEDED(Result))
+		{
+			LPWSTR IdString = nullptr;
+			if (SUCCEEDED(Device->GetId(&IdString)) && IdString)
+			{
+				OutDeviceId = FString(IdString);
+				::CoTaskMemFree(IdString);
+				return true;
+			}
+		}
+
+		return false;
+	}
+
+	bool FWasapiDeviceEnumeration::GetDeviceIndexFromId(const FString& InDeviceId, int32& OutDeviceIndex)
+	{
+		OutDeviceIndex = INDEX_NONE;
+
+		TComPtr<IMMDeviceCollection> Collection;
+		// For compatibility with the RtAudio implementation, we get all device types here
+		HRESULT Result = DeviceEnumerator->EnumAudioEndpoints(eAll, DEVICE_STATE_ACTIVE, &Collection);
+		if (FAILED(Result))
+		{
+			WASAPI_CAPTURE_LOG_RESULT("IMMDeviceCollection::EnumAudioEndpoints", Result);
+			return false;
+		}
+
+		uint32 Count;
+		Result = Collection->GetCount(&Count);
+		if (FAILED(Result))
+		{
+			WASAPI_CAPTURE_LOG_RESULT("IMMDeviceCollection::GetCount", Result);
+			return false;
+		}
+
+		for (uint32 DeviceIndex = 0; DeviceIndex < Count; ++DeviceIndex)
+		{
+			TComPtr<IMMDevice> Device;
+			Result = Collection->Item(DeviceIndex, &Device);
+			if (SUCCEEDED(Result))
+			{
+				FString DeviceId;
+				if (GetDeviceId(Device, DeviceId))
+				{
+					if (DeviceId == InDeviceId)
+					{
+						OutDeviceIndex = DeviceIndex;
+						return true;
+					}
+				}
+			}
+		}
+
+		return false;
+	}
+
+	bool FWasapiDeviceEnumeration::GetIMMDevice(const FString& InDeviceId, TComPtr<IMMDevice>& OutDevice)
+	{
+		TComPtr<IMMDevice> Device;
+		HRESULT Result = DeviceEnumerator->GetDevice(TCHAR_TO_WCHAR(*InDeviceId), &Device);
+		if (FAILED(Result))
+		{
+			WASAPI_CAPTURE_LOG_RESULT("IMMDeviceEnumerator::GetDevice", Result);
+			return false;
+		}
+
+		OutDevice = Device;
+		return true;
+	}
+
+	bool FWasapiDeviceEnumeration::GetInputDevicesAvailable(TArray<FDeviceInfo>& OutDevices)
+	{
+		TComPtr<IMMDeviceCollection> Collection;
+		HRESULT Result = DeviceEnumerator->EnumAudioEndpoints(eCapture, DEVICE_STATE_ACTIVE, &Collection);
+		if (FAILED(Result))
+		{
+			WASAPI_CAPTURE_LOG_RESULT("IMMDeviceCollection::EnumAudioEndpoints", Result);
+			return false;
+		}
+
+		uint32 Count;
+		Result = Collection->GetCount(&Count);
+		if (FAILED(Result))
+		{
+			WASAPI_CAPTURE_LOG_RESULT("IMMDeviceCollection::GetCount", Result);
+			return false;
+		}
+
+		for (uint32 DeviceIndex = 0; DeviceIndex < Count; ++DeviceIndex)
+		{
+			TComPtr<IMMDevice> Device;
+			Result = Collection->Item(DeviceIndex, &Device);
+			if (SUCCEEDED(Result))
+			{
+				FDeviceInfo DeviceInfo;
+				if (GetDeviceProperties(Device, DeviceInfo))
+				{
+					UE_LOG(LogAudioCaptureCore, Display, TEXT("Device %d: \"%s\" (%s)"), DeviceIndex,
+						DeviceInfo.FriendlyName.GetCharArray().GetData(), DeviceInfo.DeviceId.GetCharArray().GetData());
+
+					UE_LOG(LogAudioCaptureCore, Display, TEXT("\tBPS: %d, SR: %d"), DeviceInfo.BitsPerSample, DeviceInfo.PreferredSampleRate);
+
+					OutDevices.Emplace(MoveTemp(DeviceInfo));
+				}
+				else
+				{
+					return false;
+				}
+			}
+			else
+			{
+				WASAPI_CAPTURE_LOG_RESULT("IMMDeviceCollection::Item", Result);
+				return false;
+			}
+		}
+
+		return true;
+	}
+
+	bool FWasapiDeviceEnumeration::GetDeviceProperties(TComPtr<IMMDevice> InDevice, FDeviceInfo& OutInfo)
+	{
+		FString IdString;
+		FString FriendlyName;
+		if (GetDeviceId(InDevice, IdString) && GetDeviceFriendlyName(InDevice, FriendlyName))
+		{
+			TComPtr<IAudioClient3> AudioClient;
+			HRESULT Result = InDevice->Activate(__uuidof(IAudioClient3), CLSCTX_INPROC_SERVER, nullptr, IID_PPV_ARGS_Helper(&AudioClient));
+			if (FAILED(Result))
+			{
+				WASAPI_CAPTURE_LOG_RESULT("IAudioClient3::Activate", Result);
+				return false;
+			}
+
+			WAVEFORMATEX* DeviceFormat = nullptr;
+			if (GetAudioClientMixFormat(AudioClient, &DeviceFormat))
+			{
+				ON_SCOPE_EXIT{ ::CoTaskMemFree(DeviceFormat); };
+
+				OutInfo.DeviceId = IdString;
+				OutInfo.FriendlyName = FriendlyName;
+				OutInfo.BitsPerSample = DeviceFormat->wBitsPerSample;
+				OutInfo.PreferredSampleRate = DeviceFormat->nSamplesPerSec;
+
+				OutInfo.EndpointType = GetEndpointType(InDevice);
+				if (OutInfo.EndpointType == FDeviceInfo::EEndpointType::Capture)
+				{
+					OutInfo.NumInputChannels = DeviceFormat->nChannels;
+				}
+				else if (OutInfo.EndpointType == FDeviceInfo::EEndpointType::Render)
+				{
+					OutInfo.NumOutputChannels = DeviceFormat->nChannels;
+				}
+
+				return true;
+			}
+		}
+
+		return false;
+	}
+
+	bool FWasapiDeviceEnumeration::GetDeviceId(TComPtr<IMMDevice> InDevice, FString& OutString)
+	{
+		LPWSTR IdString = nullptr;
+		HRESULT Result = InDevice->GetId(&IdString);
+		if (FAILED(Result))
+		{
+			WASAPI_CAPTURE_LOG_RESULT("IMMDevice::GetId", Result);
+			return false;
+		}
+
+		OutString = FString(IdString);
+		::CoTaskMemFree(IdString);
+
+		return true;
+	}
+
+	bool FWasapiDeviceEnumeration::GetDeviceFriendlyName(TComPtr<IMMDevice> InDevice, FString& OutString)
+	{
+		TComPtr<IPropertyStore> Props;
+		HRESULT Result = InDevice->OpenPropertyStore(STGM_READ, &Props);
+		if (FAILED(Result))
+		{
+			WASAPI_CAPTURE_LOG_RESULT("IMMDevice::OpenPropertyStore", Result);
+			return false;
+		}
+
+		PROPVARIANT VarName;
+		// Initialize container for property value.
+		::PropVariantInit(&VarName);
+
+		// Get the endpoint's friendly-name property.
+		Result = Props->GetValue(PKEY_Device_FriendlyName, &VarName);
+		if (FAILED(Result))
+		{
+			WASAPI_CAPTURE_LOG_RESULT("IPropertyStore::GetValue", Result);
+			return false;
+		}
+
+		OutString = FString(VarName.pwszVal);
+		::PropVariantClear(&VarName);
+
+		return true;
+	}
+
+	FDeviceInfo::EEndpointType FWasapiDeviceEnumeration::GetEndpointType(TComPtr<IMMDevice> InDevice)
+	{
+		TComPtr<IMMEndpoint> Endpoint;
+		HRESULT Result = InDevice->QueryInterface(IID_PPV_ARGS(&Endpoint));
+		if (SUCCEEDED(Result))
+		{
+			EDataFlow DataFlow = eRender;
+			Result = Endpoint->GetDataFlow(&DataFlow);
+			if (SUCCEEDED(Result))
+			{
+				switch (DataFlow)
+				{
+				case eRender:
+					return FDeviceInfo::EEndpointType::Render;
+				case eCapture:
+					return FDeviceInfo::EEndpointType::Capture;
+				default:
+					break;
+				}
+			}
+			else
+			{
+				WASAPI_CAPTURE_LOG_RESULT("IMMEndpoint::GetDataFlow", Result);
+			}
+		}
+		else
+		{
+			WASAPI_CAPTURE_LOG_RESULT("IMMDevice::QueryInterface", Result);
+		}
+
+		return FDeviceInfo::EEndpointType::Unknown;
+	}
+
+	bool FWasapiDeviceEnumeration::GetAudioClientMixFormat(TComPtr<IAudioClient3> AudioClient, WAVEFORMATEX** OutFormat)
+	{
+		HRESULT Result = AudioClient->GetMixFormat(OutFormat);
+		if (FAILED(Result))
+		{
+			WASAPI_CAPTURE_LOG_RESULT("IAudioClient3::GetMixFormat", Result);
+			return false;
+		}
+
+		return true;
+	}
+}
\ No newline at end of file
diff --git a/Engine/Source/Runtime/AudioCaptureImplementations/Windows/AudioCaputureWasapi/Private/WasapiDeviceEnumeration.h b/Engine/Source/Runtime/AudioCaptureImplementations/Windows/AudioCaputureWasapi/Private/WasapiDeviceEnumeration.h
new file mode 100644
--- /dev/null
+++ b/Engine/Source/Runtime/AudioCaptureImplementations/Windows/AudioCaputureWasapi/Private/WasapiDeviceEnumeration.h
@@ -0,0 +1,144 @@
+// Copyright Epic Games, Inc. All Rights Reserved.
+
+#pragma once
+
+#include "AudioCaptureCore.h"
+
+#include "Windows/AllowWindowsPlatformTypes.h"
+#include "Windows/AllowWindowsPlatformAtomics.h"
+
+#include "Microsoft/COMPointer.h"
+
+THIRD_PARTY_INCLUDES_START
+#include <AudioClient.h>
+#include <Mmdeviceapi.h>
+THIRD_PARTY_INCLUDES_END
+
+#include "Windows/HideWindowsPlatformAtomics.h"
+#include "Windows/HideWindowsPlatformTypes.h"
+
+
+namespace Audio
+{
+	/**
+	 * FWasapiDeviceEnumeration - Class which wraps Microsoft's IMM device enumeration API.
+	 */
+	class FWasapiDeviceEnumeration
+	{
+	public:
+		/**
+		 * FDeviceInfo - Properties associated with an audio device.
+		 */
+		struct FDeviceInfo
+		{
+			/**
+			 * EEndpointType - Designates whether a device is a capture or render device.
+			 */
+			enum class EEndpointType { Unknown, Render, Capture };
+
+			/** Unique identifier for this device. */
+			FString DeviceId;
+			/** Human readable name for this device. */
+			FString FriendlyName;
+			/** Number of input channels on this device available for recording. */
+			uint32 NumInputChannels = 0;
+			/** Number of output channels on this device available for playback. */
+			uint32 NumOutputChannels = 0;
+			/** The current bit depth configured for this device. Additional bit depths may be supported. */
+			uint32 BitsPerSample = 0;
+			/** The preferred sample rate for this device. Additional sample rates may be supported. */
+			uint32 PreferredSampleRate = 0;
+			/** The endpoint type for this device (Render or Capture). */
+			EEndpointType EndpointType = EEndpointType::Unknown;
+		};
+
+		FWasapiDeviceEnumeration() = default;
+		~FWasapiDeviceEnumeration() = default;
+
+		/** Instantiates DeviceEnumerator and initializes the default input and output devices IDs.  */
+		void Initialize();
+
+		/** Returns the Id of the default input device. */
+		FString GetDefaultInputDeviceId();
+		/** Returns the Id of the default output device. */
+		FString GetDefaultOutputDeviceId();
+
+		/**
+		 * GetDeviceInfo - Gets the FDeviceInfo for the device with the given Id.
+		 * 
+		 * @param InDeviceId - The Id of the device to fetch information for.
+		 * @param OutDeviceInfo - Upon success, the device info struct for the device. 
+		 * Untouched on failure.
+		 * @return - Boolean indicating success or failure when attempting to fetch the device info.
+		 */
+		bool GetDeviceInfo(const FString& InDeviceId, FDeviceInfo& OutDeviceInfo);
+
+		/**
+		 * GetDeviceIdFromIndex - Gets the device Id from a device index.
+		 * 
+		 * @param InDeviceIndex - The index for the device to get the Id for.
+		 * @param InDataFlow - Indicates which device collection to search (render, capture or both).
+		 * @param OutDeviceId - Upon success, the Id of the device at the given index.
+		 * @return - Boolean indicating success or failure when looking up the device index.
+		 */
+		bool GetDeviceIdFromIndex(int32 InDeviceIndex, EDataFlow InDataFlow, FString& OutDeviceId);
+		
+		/**
+		 * GetDeviceIndexFromId - Gets the index of a device with the given Id. Searches the
+		 * collection of active devices for both capture and render data flows. 
+		 * 
+		 * @param InDeviceId - Device Id to get the index for.
+		 * @param OutDeviceIndex - Upon success, will contain the index for the device.
+		 * @return - Boolean indicating success or failure when looking up the device Id.
+		 */
+		bool GetDeviceIndexFromId(const FString& InDeviceId, int32& OutDeviceIndex);
+
+		/**
+		 * GetInputDevicesAvailable - Gets the collection of currently active audio capture devices.
+		 * 
+		 * @param OutDevices - TArray of FDeviceInfo objects for the currently available capture devices.
+		 * @return - Boolean indicating success or failure when access the collection of devices.
+		 */
+		bool GetInputDevicesAvailable(TArray<FDeviceInfo>& OutDevices);
+
+		/**
+		 * GetIMMDevice - Gets a COM pointer to the IMMDevice for the give device Id.
+		 * 
+		 * @param InDeviceId - Device Id of the device to get.
+		 * @param OutDevice - COM pointer to the device for the given Id.
+		 * @return - Boolean indicating success or failure when looking up the give Id.
+		 */
+		bool GetIMMDevice(const FString& InDeviceId, TComPtr<IMMDevice>& OutDevice);
+
+	private:
+		/** The Id of the system default render device. */
+		FString DefaultRenderId;
+		/** The Id of the system default capture device. */
+		FString DefaultCaptureId;
+		/** COM pointer to the device enumerator used by this class to find audio devices and their properties. */
+		TComPtr<IMMDeviceEnumerator> DeviceEnumerator;
+
+		/** EnumerateDefaults - Get the IDs for the system default devices. */
+		void EnumerateDefaults();
+
+		/** GetDeviceProperties - Convenience function for getting the device info of an IMMDevice. */
+		bool GetDeviceProperties(TComPtr<IMMDevice> InDevice, FDeviceInfo& OutInfo);
+
+		/** GetDeviceId - Convenience function for getting the device Id of an IMMDevice. */
+		bool GetDeviceId(TComPtr<IMMDevice> InDevice, FString& OutString);
+
+		/** GetDeviceFriendlyName - Convenience function for getting the friendly name of an IMMDevice. */
+		bool GetDeviceFriendlyName(TComPtr<IMMDevice> InDevice, FString& OutString);
+
+		/** GetEndpointType - Returns endpoint type (render or capture) for the given IMMDevice. */
+		FDeviceInfo::EEndpointType GetEndpointType(TComPtr<IMMDevice> InDevice);
+
+		/**
+		 * GetAudioClientMixFormat - Returns the current mix format for the given audio client.
+		 * The mix format is the audio format used by the OS audio mixer when mixing shared
+		 * mode streams.
+		 */
+		bool GetAudioClientMixFormat(TComPtr<IAudioClient3> AudioClient, WAVEFORMATEX** OutFormat);
+	};
+
+}
diff --git a/Engine/Source/Runtime/AudioCaptureImplementations/Windows/AudioCaputureWasapi/Private/WasapiInputStream.cpp b/Engine/Source/Runtime/AudioCaptureImplementations/Windows/AudioCaputureWasapi/Private/WasapiInputStream.cpp
new file mode 100644
--- /dev/null
+++ b/Engine/Source/Runtime/AudioCaptureImplementations/Windows/AudioCaputureWasapi/Private/WasapiInputStream.cpp
@@ -0,0 +1,252 @@
+// Copyright Epic Games, Inc. All Rights Reserved.
+
+#include "WasapiInputStream.h"
+
+#include "WasapiCaptureLog.h"
+
+namespace Audio
+{
+	// REFERENCE_TIME base which is in units of 100 nanoseconds
+	static constexpr REFERENCE_TIME ReferenceTimeBase = 1e7;
+
+	FWasapiInputStream::FWasapiInputStream(
+		TComPtr<IMMDevice> InDevice, 
+		const FWasapiAudioFormat& InFormat, 
+		uint32 InNumFramesDesired,
+		FWasapiOnAudioCaptureFunction InCallback)
+	{
+		check(InDevice.IsValid());
+
+		// Activating the device, which gives us the audio client.
+		TComPtr<IAudioClient3> TempAudioClient;
+		HRESULT Result = InDevice->Activate(__uuidof(IAudioClient3), CLSCTX_INPROC_SERVER, nullptr, IID_PPV_ARGS_Helper(&TempAudioClient));
+		if (FAILED(Result))
+		{
+			WASAPI_CAPTURE_LOG_RESULT("IMMDevice::Activate", Result);
+			return;
+		}
+
+		WAVEFORMATEX* MixFormat = nullptr;
+		Result = TempAudioClient->GetMixFormat(&MixFormat);
+		if (FAILED(Result))
+		{
+			WASAPI_CAPTURE_LOG_RESULT("IAudioClient3::GetMixFormat", Result);
+			return;
+		}
+
+		REFERENCE_TIME MinimumPeriodInRefTimes = 0;
+		Result = TempAudioClient->GetDevicePeriod(nullptr, &MinimumPeriodInRefTimes);
+		if (FAILED(Result))
+		{
+			WASAPI_CAPTURE_LOG_RESULT("IAudioClient3::GetDevicePeriod", Result);
+			return;
+		}
+
+		// Audio events will be delivered to us rather than needing to poll
+		uint32 Flags = AUDCLNT_STREAMFLAGS_EVENTCALLBACK;
+		// Convert desired buffer size from frames to REFERENCE_TIME which is in units of 100 nanoseconds
+		double BufferDurationSeconds = static_cast<double>(InNumFramesDesired) / InFormat.GetSampleRate();
+		REFERENCE_TIME DesiredBufferDuration = static_cast<REFERENCE_TIME>(BufferDurationSeconds * ReferenceTimeBase);
+		// Ensure the requested duration is at least as large as the minimum supported by this device
+		DesiredBufferDuration = FMath::Max(MinimumPeriodInRefTimes, DesiredBufferDuration);
+
+		Result = TempAudioClient->Initialize(AUDCLNT_SHAREMODE_SHARED, Flags, DesiredBufferDuration, 0, InFormat.GetWaveFormat(), nullptr);
+		if (FAILED(Result))
+		{
+			WASAPI_CAPTURE_LOG_RESULT("IAudioClient3::Initialize", Result);
+			return;
+		}
+
+		Result = TempAudioClient->GetBufferSize(&NumFramesPerBuffer);
+		if (FAILED(Result))
+		{
+			WASAPI_CAPTURE_LOG_RESULT("IAudioClient3::GetBufferSize", Result);
+			return;
+		}
+
+		// Not using FEvent/FEventWin here because we need access to the raw, platform
+		// handle (see SetEventHandler() below).
+		EventHandle = ::CreateEvent(nullptr, 0, 0, nullptr);
+		if (EventHandle == nullptr)
+		{
+			WASAPI_CAPTURE_LOG_RESULT("CreateEvent", Result);
+			return;
+		}
+
+		Result = TempAudioClient->SetEventHandle(EventHandle);
+		if (FAILED(Result))
+		{
+			WASAPI_CAPTURE_LOG_RESULT("IAudioClient3::SetEventHandle", Result);
+			return;
+		}
+
+		TComPtr<IAudioCaptureClient> TempCaptureClient;
+		Result = TempAudioClient->GetService(__uuidof(IAudioCaptureClient), IID_PPV_ARGS_Helper(&TempCaptureClient));
+		if (FAILED(Result))
+		{
+			WASAPI_CAPTURE_LOG_RESULT("IAudioClient3::GetService", Result);
+			return;
+		}
+
+		AudioClient = MoveTemp(TempAudioClient);
+		CaptureClient = MoveTemp(TempCaptureClient);
+		AudioFormat = InFormat;
+		OnAudioCaptureCallback = MoveTemp(InCallback);
+		SilienceBuffer.SetNumZeroed(GetBufferSizeBytes(), false);
+
+		bIsInitialized = true;
+
+		UE_LOG(LogAudioCaptureCore, Display, TEXT("WasapiCapture AudioFormat SampeRate: %d, BitDepth: %s"), 
+			AudioFormat.GetSampleRate(), *(AudioFormat.GetEncodingString()));
+	}
+
+	FWasapiInputStream::~FWasapiInputStream()
+	{
+		if (EventHandle)
+		{
+			::CloseHandle(EventHandle);
+		}
+	}
+
+	bool FWasapiInputStream::IsInitialized() const
+	{
+		return bIsInitialized;
+	}
+
+	void FWasapiInputStream::StartStream()
+	{
+		if (AudioClient)
+		{
+			AudioClient->Start();
+
+			// Drain the CaptureClient input buffer
+			// see comment here for more info: 
+			// https://learn.microsoft.com/en-us/windows/win32/api/audioclient/nf-audioclient-iaudiocaptureclient-getbuffer
+			DrainInputBuffer();
+		}
+	}
+
+	void FWasapiInputStream::StopStream()
+	{ 
+		if (AudioClient)
+		{
+			AudioClient->Stop();
+		}
+	}
+
+	bool FWasapiInputStream::WaitOnBuffer() const
+	{
+		static constexpr uint32 TimeoutInMs = 1000;
+
+		if (::WaitForSingleObject(EventHandle, TimeoutInMs) != WAIT_OBJECT_0)
+		{
+			return false;
+		}
+
+		return true;
+	}
+
+	uint32 FWasapiInputStream::GetBufferSizeFrames() const
+	{
+		return NumFramesPerBuffer;
+	}
+
+	uint32 FWasapiInputStream::GetBufferSizeBytes() const
+	{
+		return NumFramesPerBuffer * AudioFormat.GetNumChannels() * AudioFormat.GetBytesPerSample();
+	}
+
+	double FWasapiInputStream::GetStreamPosition() const
+	{
+		return DevicePosition;
+	}
+
+	bool FWasapiInputStream::DrainInputBuffer()
+	{
+		uint32 NextPacketSize = 0;
+		HRESULT Result = CaptureClient->GetNextPacketSize(&NextPacketSize);
+		if (FAILED(Result))
+		{
+			// Don't log from the audio thread.
+			return false;
+		}
+
+		while (NextPacketSize > 0)
+		{
+			DWORD Flags = 0;
+			uint8* Data = nullptr;
+			uint32 NumFramesRead = 0;
+
+			Result = CaptureClient->GetBuffer(&Data, &NumFramesRead, &Flags, nullptr, nullptr);
+			if (FAILED(Result))
+			{
+				return false;
+			}
+
+			CaptureClient->ReleaseBuffer(NumFramesRead);
+
+			Result = CaptureClient->GetNextPacketSize(&NextPacketSize);
+			if (FAILED(Result))
+			{
+				return false;
+			}
+		}
+
+		return true;
+	}
+
+	bool FWasapiInputStream::CaptureAudioFrames()
+	{
+		SCOPED_NAMED_EVENT(FWasapiInputStream_CaptureAudioFrames, FColor::Red);
+
+		uint32 NextPacketSize = 0;
+		HRESULT Result = CaptureClient->GetNextPacketSize(&NextPacketSize);
+		if (FAILED(Result))
+		{
+			// We aren't logging failures here because this method is called in the
+			// audio thread context.
+			return false;
+		}
+
+		DWORD Flags = 0;
+		uint8* Data = nullptr;
+		uint32 NumFramesRead = 0;
+		uint64 TempDevicePosition = 0;
+
+		while (NextPacketSize > 0)
+		{
+			Result = CaptureClient->GetBuffer(&Data, &NumFramesRead, &Flags, &TempDevicePosition, nullptr);
+			if (FAILED(Result))
+			{
+				return false;
+			}
+
+			DevicePosition.store(static_cast<double>(TempDevicePosition) / AudioFormat.GetSampleRate(), std::memory_order_relaxed);
+			bool bHasDiscontinuityError = ((Flags & AUDCLNT_BUFFERFLAGS_DATA_DISCONTINUITY) != 0);
+
+			// If this flag is set we need to ignore the data returned by GetBuffer()
+			// and treat it as silence. We utilize SilienceBuffer for this.
+			if ((Flags & AUDCLNT_BUFFERFLAGS_SILENT) != 0)
+			{
+				int32 BytesRead = NumFramesRead * AudioFormat.GetNumChannels() * AudioFormat.GetBytesPerSample();
+				check(SilienceBuffer.Max() >= BytesRead);
+				Data = SilienceBuffer.GetData();
+			}
+
+			if (OnAudioCaptureCallback != nullptr)
+			{
+				OnAudioCaptureCallback(Data, NumFramesRead, DevicePosition, bHasDiscontinuityError);
+			}
+
+			CaptureClient->ReleaseBuffer(NumFramesRead);
+
+			Result = CaptureClient->GetNextPacketSize(&NextPacketSize);
+			if (FAILED(Result))
+			{
+				return false;
+			}
+		}
+
+		return true;
+	}
+}
diff --git a/Engine/Source/Runtime/AudioCaptureImplementations/Windows/AudioCaputureWasapi/Private/WasapiInputStream.h b/Engine/Source/Runtime/AudioCaptureImplementations/Windows/AudioCaputureWasapi/Private/WasapiInputStream.h
new file mode 100644
--- /dev/null
+++ b/Engine/Source/Runtime/AudioCaptureImplementations/Windows/AudioCaputureWasapi/Private/WasapiInputStream.h
@@ -0,0 +1,103 @@
+// Copyright Epic Games, Inc. All Rights Reserved.
+
+#pragma once
+
+#include "AudioCaptureCore.h"
+#include "WasapiAudioFormat.h"
+
+#include "Windows/AllowWindowsPlatformTypes.h"
+#include "Windows/AllowWindowsPlatformAtomics.h"
+
+#include "Microsoft/COMPointer.h"
+
+THIRD_PARTY_INCLUDES_START
+#include <AudioClient.h>
+#include <comdef.h>
+#include <Mmdeviceapi.h>
+THIRD_PARTY_INCLUDES_END
+
+#include "Windows/HideWindowsPlatformAtomics.h"
+#include "Windows/HideWindowsPlatformTypes.h"
+
+
+namespace Audio
+{
+	/**
+	 * FWasapiOnAudioCaptureFunction - Callback function type for receiving audio from a capture device.
+	 */
+	typedef TFunction<void(void* InBuffer, uint32 InNumFrames, double InStreamPosition, bool bInDiscontinuityError)> FWasapiOnAudioCaptureFunction;
+
+	/**
+	 * FWasapiInputStream - This class manages the WASPI audio client, capture buffers and audio callback used during capture.
+	 */
+	class FWasapiInputStream
+	{
+	public:
+		FWasapiInputStream() = delete;
+		FWasapiInputStream(FWasapiInputStream&& InOther) = delete;
+		FWasapiInputStream(const FWasapiInputStream& InOther) = delete;
+
+		FWasapiInputStream(
+			TComPtr<IMMDevice> InDevice, 
+			const FWasapiAudioFormat& InFormat, 
+			uint32 InNumFramesDesired,
+			FWasapiOnAudioCaptureFunction InCallback);
+
+		virtual ~FWasapiInputStream();
+
+		FWasapiInputStream& operator=(FWasapiInputStream&& InOther) = delete;
+		FWasapiInputStream& operator=(const FWasapiInputStream& InOther) = delete;
+
+		/** Returns whether or not this object has been successfully initialized. */
+		bool IsInitialized() const;
+
+		/** Starts the audio capture which triggers periodic callbacks with new audio data. */
+		void StartStream();
+		/** Stops the audio client which will shut down the stream and any further callbacks. */
+		void StopStream();
+
+		/** Waits on the audio client event handle which will get set when new capture data is available. */
+		bool WaitOnBuffer() const;
+		/** The capture thread calls this periodically with new audio data during a capture. */
+		bool CaptureAudioFrames();
+
+		/** Returns the buffer size in frames for the buffer for the audio callback. */
+		uint32 GetBufferSizeFrames() const;
+		/** Returns the buffer size in bytes corresponding to the buffer for the audio callback. */
+		uint32 GetBufferSizeBytes() const;
+		/** Returns the current stream position in seconds. */
+		double GetStreamPosition() const;
+
+	private:
+
+		/** Indicates if this object has been successfully initialized. */
+		bool bIsInitialized = false;
+		/** COM pointer to the WASAPI audio client object. */
+		TComPtr<IAudioClient3> AudioClient;
+		/** COM pointer to the WASAPI capture client object. */
+		TComPtr<IAudioCaptureClient> CaptureClient;
+
+		/** Holds the audio format configuration for this stream. */
+		FWasapiAudioFormat AudioFormat;
+		/** The callback used to periodically deliver new audio data during capture. */
+		FWasapiOnAudioCaptureFunction OnAudioCaptureCallback;
+		/** Buffer used when WASAPI indicates that silence should be output for a given callback. */
+		TArray<uint8> SilienceBuffer;
+
+		/** Number of frames of audio data which will be used for each audio callback during capture. */
+		uint32 NumFramesPerBuffer = 0;
+		/** Current device position as reported by WASAPI each callback. */
+		std::atomic<double> DevicePosition = 0.0;
+		/**
+		 * Event handle which our capture thread waits on prior to each callback. WASAPI signals this 
+		 * object each quanta when a buffer of audio has been captured and is ready to be consumed downstream.
+		 */
+		HANDLE EventHandle = nullptr;
+
+		/**
+		 * DrainInputBuffer - Prepares the capture client by draining the input buffer of any audio data
+		 * which may have accumulated.
+		 */
+		bool DrainInputBuffer();
+	};
+}
diff --git a/Engine/Source/Runtime/AudioCaptureImplementations/Windows/AudioCaputureWasapi/Private/WasapiStreamManager.cpp b/Engine/Source/Runtime/AudioCaptureImplementations/Windows/AudioCaputureWasapi/Private/WasapiStreamManager.cpp
new file mode 100644
--- /dev/null
+++ b/Engine/Source/Runtime/AudioCaptureImplementations/Windows/AudioCaputureWasapi/Private/WasapiStreamManager.cpp
@@ -0,0 +1,151 @@
+// Copyright Epic Games, Inc. All Rights Reserved.
+
+#include "WasapiStreamManager.h"
+
+#include "AudioCaptureCoreLog.h"
+#include "Misc/ScopeExit.h"
+
+
+namespace Audio
+{
+	FWasapiStreamManager::FWasapiStreamManager()
+	{
+		const bool bResult = FWindowsPlatformMisc::CoInitialize();
+		if (bResult)
+		{
+			bCoInitialized = true;
+		}
+
+		DeviceEnumerator.Initialize();
+	}
+
+	FWasapiStreamManager::~FWasapiStreamManager()
+	{
+		if (bCoInitialized)
+		{
+			FWindowsPlatformMisc::CoUninitialize();
+		}
+	}
+
+	FString FWasapiStreamManager::GetDefaultInputDeviceId()
+	{
+		return DeviceEnumerator.GetDefaultInputDeviceId();
+	}
+
+	FString FWasapiStreamManager::GetDefaultOutputDeviceId()
+	{
+		return DeviceEnumerator.GetDefaultOutputDeviceId();
+	}
+
+	bool FWasapiStreamManager::GetDeviceInfo(const FString& InDeviceId, FDeviceInfo& OutDeviceInfo)
+	{
+		return DeviceEnumerator.GetDeviceInfo(InDeviceId, OutDeviceInfo);
+	}
+
+	bool FWasapiStreamManager::GetDeviceIndexFromId(const FString& InDeviceId, int32& OutDeviceIndex)
+	{
+		return DeviceEnumerator.GetDeviceIndexFromId(InDeviceId, OutDeviceIndex);
+	}
+
+	bool FWasapiStreamManager::GetDeviceIdFromIndex(int32 InDeviceIndex, EDataFlow InDataFlow, FString& OutDeviceId)
+	{
+		return DeviceEnumerator.GetDeviceIdFromIndex(InDeviceIndex, InDataFlow, OutDeviceId);
+	}
+
+	bool FWasapiStreamManager::GetInputDevicesAvailable(TArray<FDeviceInfo>& OutDevices)
+	{
+		return DeviceEnumerator.GetInputDevicesAvailable(OutDevices);
+	}
+
+	bool FWasapiStreamManager::OpenStream(
+		const FString& InDeviceId, 
+		const FWasapiAudioFormat& InFormat,
+		uint32 InNumFramesDesired,
+		FWasapiOnAudioCaptureFunction InCallback)
+	{
+		check(State == EStreamState::STREAM_CLOSED);
+		check(InputStream == nullptr);
+
+		TComPtr<IMMDevice> Device;
+		if (DeviceEnumerator.GetIMMDevice(InDeviceId, Device))
+		{
+			InputStream = MakeShared<FWasapiInputStream>(Device, InFormat, InNumFramesDesired, MoveTemp(InCallback));
+			if (InputStream->IsInitialized())
+			{
+				CaptureThread = MakeUnique<FWasapiCaptureThread>(InputStream);
+				State = EStreamState::STREAM_STOPPED;
+
+				return true;
+			}
+			else
+			{
+				InputStream = nullptr;
+			}
+		}
+
+		return false;
+	}
+
+	bool FWasapiStreamManager::IsStreamOpen() const
+	{
+		return State != EStreamState::STREAM_CLOSED;
+	}
+
+	void FWasapiStreamManager::StartStream()
+	{
+		check(State == EStreamState::STREAM_STOPPED);
+
+		if (CaptureThread->Start())
+		{
+			State = EStreamState::STREAM_CAPTURING;
+		}
+		else
+		{
+			State = EStreamState::INVALID_STATE;
+		}
+	}
+
+	bool FWasapiStreamManager::IsCapturing() const
+	{
+		return State == EStreamState::STREAM_CAPTURING;
+	}
+
+	void FWasapiStreamManager::StopStream()
+	{
+		check(State == EStreamState::STREAM_CAPTURING);
+
+		State = EStreamState::STREAM_STOPPING;
+		CaptureThread->Stop();
+		State = EStreamState::STREAM_STOPPED;
+	}
+
+	void FWasapiStreamManager::AbortStream()
+	{
+		State = EStreamState::STREAM_STOPPING;
+		CaptureThread->Abort();
+		State = EStreamState::STREAM_STOPPED;
+	}
+
+	void FWasapiStreamManager::CloseStream()
+	{
+		if (State == EStreamState::STREAM_CAPTURING)
+		{
+			StopStream();
+		}
+
+		CaptureThread = nullptr;
+		InputStream = nullptr;
+
+		State = EStreamState::STREAM_CLOSED;
+	}
+
+	uint32 FWasapiStreamManager::GetStreamBufferSizeBytes() const
+	{
+		return InputStream->GetBufferSizeBytes();
+	}
+
+	double FWasapiStreamManager::GetStreamPosition() const
+	{
+		return InputStream->GetStreamPosition();
+	}
+}
diff --git a/Engine/Source/Runtime/AudioCaptureImplementations/Windows/AudioCaputureWasapi/Private/WasapiStreamManager.h b/Engine/Source/Runtime/AudioCaptureImplementations/Windows/AudioCaputureWasapi/Private/WasapiStreamManager.h
new file mode 100644
--- /dev/null
+++ b/Engine/Source/Runtime/AudioCaptureImplementations/Windows/AudioCaputureWasapi/Private/WasapiStreamManager.h
@@ -0,0 +1,119 @@
+// Copyright Epic Games, Inc. All Rights Reserved.
+
+#pragma once
+
+#include "AudioCaptureCore.h"
+#include "WasapiDeviceEnumeration.h"
+#include "WasapiInputStream.h"
+#include "WasapiCaptureThread.h"
+
+
+namespace Audio
+{
+	/**
+	 * FWasapiStreamManager - Top level object for managing audio captures. It is composed of
+	 * a device enumerator, capture thread, input stream and state machine. In the future, 
+	 * additional input streams may be managed here in support of multi-track capture.
+	 */
+	class FWasapiStreamManager
+	{
+	public:
+
+		using FDeviceInfo = FWasapiDeviceEnumeration::FDeviceInfo;
+
+		FWasapiStreamManager();
+		FWasapiStreamManager(FWasapiStreamManager&& InOther) = delete;
+		FWasapiStreamManager(const FWasapiStreamManager& InOther) = delete;
+
+		virtual ~FWasapiStreamManager();
+
+		FWasapiStreamManager& operator=(FWasapiStreamManager&& InOther) = delete;
+		FWasapiStreamManager& operator=(const FWasapiStreamManager& InOther) = delete;
+
+		/** Returns the default input device Id. */
+		FString GetDefaultInputDeviceId();
+		/** Returns the default output device Id. */
+		FString GetDefaultOutputDeviceId();
+
+		/**
+		 * GetDeviceInfo - Convienence method for FWasapiDeviceEnumeration::GetDeviceInfo()
+		 */
+		bool GetDeviceInfo(const FString& InDeviceId, FDeviceInfo& OutDeviceInfo);
+
+		/**
+		 * GetDeviceIndexFromId - Convienence method for FWasapiDeviceEnumeration::GetDeviceIndexFromId()
+		 */
+		bool GetDeviceIndexFromId(const FString& InDeviceId, int32& OutDeviceIndex);
+
+		/**
+		 * GetDeviceIdFromIndex - Convienence method for FWasapiDeviceEnumeration::GetDeviceIdFromIndex()
+		 */
+		bool GetDeviceIdFromIndex(int32 InDeviceIndex, EDataFlow InDataFlow, FString& OutDeviceId);
+
+		/**
+		 * GetInputDevicesAvailable - Convienence method for FWasapiDeviceEnumeration::GetInputDevicesAvailable()
+		 */
+		bool GetInputDevicesAvailable(TArray<FDeviceInfo>& OutDevices);
+
+		/**
+		 * OpenStream - Finds the device for the given Id and prepares it for capturing audio. It also
+		 * spawns a capture thread.
+		 * 
+		 * @param InDeviceId - Id of the device to use for capturing audio.
+		 * @param InFormat - The audio format to use for the capture.
+		 * @param InNumFramesDesired - The number of audio frames desired when audio is delivered via 
+		 * the capture callback.
+		 * @param InCallback - The callback which is called periodically during capture with audio data.
+		 */
+		bool OpenStream(
+			const FString& InDeviceId, 
+			const FWasapiAudioFormat& InFormat,
+			uint32 InNumFramesDesired,
+			FWasapiOnAudioCaptureFunction InCallback);
+
+		/** Returns whether or not this object has been successfully initialized. */
+		bool IsStreamOpen() const;
+
+		/** Starts the audio capture which triggers periodic callbacks with new audio data. */
+		void StartStream();
+		/** Indicates if the stream was successfully started and is capturing data. */
+		bool IsCapturing() const;
+		/** Stops the capture thread which, in turn, completes the runnable and stops the input stream. */
+		void StopStream();
+
+		/** Brute force stops the capture thread. */
+		void AbortStream();
+		/** Releases resources related to the capture thread and input streams. */
+		void CloseStream();
+
+		/** Returns the buffer size in bytes of the callback buffer. */
+		uint32 GetStreamBufferSizeBytes() const;
+		/** Returns the current position (in seconds) of the input stream. */
+		double GetStreamPosition() const;
+
+	private:
+		/** The states used by the state machine. */
+		enum class EStreamState {
+			INVALID_STATE,
+			STREAM_CLOSED,
+			STREAM_STOPPED,
+			STREAM_STOPPING,
+			STREAM_CAPTURING
+		};
+
+		/** Indicates if FWindowsPlatformMisc::CoInitialize() was successfull. */
+		bool bCoInitialized = false;
+
+		/** Current state of the state machine. */
+		EStreamState State = EStreamState::STREAM_CLOSED;
+
+		/** Object used to enumerate audio devices present in the system. */
+		FWasapiDeviceEnumeration DeviceEnumerator;
+
+		/** The input stream for the selected device which will stream audio during capture. */
+		TSharedPtr<FWasapiInputStream> InputStream;
+
+		/** The thread which provides an execution context during audio capture. */
+		TUniquePtr<FWasapiCaptureThread> CaptureThread;
+	};
+}
