Fixing static analysis warning.

diff --git a/Engine/Source/Runtime/AudioCaptureImplementations/Windows/AudioCaputureWasapi/Private/WasapiDeviceEnumeration.cpp b/Engine/Source/Runtime/AudioCaptureImplementations/Windows/AudioCaputureWasapi/Private/WasapiDeviceEnumeration.cpp
--- a/Engine/Source/Runtime/AudioCaptureImplementations/Windows/AudioCaputureWasapi/Private/WasapiDeviceEnumeration.cpp
+++ b/Engine/Source/Runtime/AudioCaptureImplementations/Windows/AudioCaputureWasapi/Private/WasapiDeviceEnumeration.cpp
@@ -25,9 +25,8 @@ namespace Audio
 		// Assumes FWindowsPlatformMisc::CoInitialize() has been called upstream
 		ensure(
 			SUCCEEDED(
-				::CoCreateInstance(__uuidof(MMDeviceEnumerator), nullptr, CLSCTX_INPROC_SERVER, IID_PPV_ARGS(&DeviceEnumerator)) &&
-				DeviceEnumerator != nullptr
-			)
+				::CoCreateInstance(__uuidof(MMDeviceEnumerator), nullptr, CLSCTX_INPROC_SERVER, IID_PPV_ARGS(&DeviceEnumerator))
+			) && (DeviceEnumerator != nullptr)
 		);
 
 		EnumerateDefaults();
