Fixing log level.

diff --git a/Engine/Source/Runtime/AudioCaptureImplementations/Windows/AudioCaputureWasapi/Private/WasapiDeviceEnumeration.cpp b/Engine/Source/Runtime/AudioCaptureImplementations/Windows/AudioCaputureWasapi/Private/WasapiDeviceEnumeration.cpp
--- a/Engine/Source/Runtime/AudioCaptureImplementations/Windows/AudioCaputureWasapi/Private/WasapiDeviceEnumeration.cpp
+++ b/Engine/Source/Runtime/AudioCaptureImplementations/Windows/AudioCaputureWasapi/Private/WasapiDeviceEnumeration.cpp
@@ -56,22 +56,22 @@ namespace Audio
 		FString DeviceIdName;
 		if (GetDefaultDeviceID(eRender, eConsole, DeviceIdName))
 		{
-			UE_LOG(LogAudioCaptureCore, Verbose, TEXT("WasapiAbstractionLayer: Default Device='%s'"), *DeviceIdName);
+			UE_LOG(LogAudioCaptureCore, Verbose, TEXT("FWasapiDeviceEnumeration: Default Device='%s'"), *DeviceIdName);
 			DefaultRenderId = DeviceIdName;
 		}
 		else
 		{
-			UE_LOG(LogAudioCaptureCore, Error, TEXT("WasapiAbstractionLayer: error fetching default render device"));
+			UE_LOG(LogAudioCaptureCore, Verbose, TEXT("FWasapiDeviceEnumeration: no default render device found"));
 		}
 
 		if (GetDefaultDeviceID(eCapture, eConsole, DeviceIdName))
 		{
-			UE_LOG(LogAudioCaptureCore, Verbose, TEXT("WasapiAbstractionLayer: Default Device='%s'"), *DeviceIdName);
+			UE_LOG(LogAudioCaptureCore, Verbose, TEXT("FWasapiDeviceEnumeration: Default Device='%s'"), *DeviceIdName);
 			DefaultCaptureId = DeviceIdName;
 		}
 		else
 		{
-			UE_LOG(LogAudioCaptureCore, Error, TEXT("WasapiAbstractionLayer: error fetching default capture device"));
+			UE_LOG(LogAudioCaptureCore, Verbose, TEXT("FWasapiDeviceEnumeration: no default capture device found"));
 		}
 	}
 
