During audio device enumeration, log errors and contine instead of returning empty list.

diff --git a/Engine/Source/Runtime/AudioCaptureImplementations/Windows/AudioCaputureWasapi/Private/WasapiDeviceEnumeration.cpp b/Engine/Source/Runtime/AudioCaptureImplementations/Windows/AudioCaputureWasapi/Private/WasapiDeviceEnumeration.cpp
--- a/Engine/Source/Runtime/AudioCaptureImplementations/Windows/AudioCaputureWasapi/Private/WasapiDeviceEnumeration.cpp
+++ b/Engine/Source/Runtime/AudioCaptureImplementations/Windows/AudioCaputureWasapi/Private/WasapiDeviceEnumeration.cpp
@@ -219,13 +219,12 @@ namespace Audio
 				}
 				else
 				{
-					return false;
+					UE_LOG(LogAudioCaptureCore, Warning, TEXT("Unable to fetch audio capture device properties...skipping"));
 				}
 			}
 			else
 			{
-				WASAPI_CAPTURE_LOG_RESULT("IMMDeviceCollection::Item", Result);
-				return false;
+				UE_LOG(LogAudioCaptureCore, Warning, TEXT("Unable to fetch audio capture device...skipping"));
 			}
 		}
 
