Fix empty audio device list in Mac Take Recorder.

diff --git a/Engine/Source/Runtime/AudioCaptureImplementations/AudioCaptureRtAudio/Private/AudioCaptureRtAudio.cpp b/Engine/Source/Runtime/AudioCaptureImplementations/AudioCaptureRtAudio/Private/AudioCaptureRtAudio.cpp
--- a/Engine/Source/Runtime/AudioCaptureImplementations/AudioCaptureRtAudio/Private/AudioCaptureRtAudio.cpp
+++ b/Engine/Source/Runtime/AudioCaptureImplementations/AudioCaptureRtAudio/Private/AudioCaptureRtAudio.cpp
@@ -83,6 +83,13 @@ bool Audio::FAudioCaptureRtAudioStream::GetCaptureDeviceInfo(FCaptureDeviceInfo&
 		OutInfo.DeviceId = FString(FUTF8ToTCHAR(DeviceInfo.deviceId.c_str()));
 		OutInfo.InputChannels = DeviceInfo.inputChannels;
 		OutInfo.PreferredSampleRate = DeviceInfo.preferredSampleRate;
+        
+        if (OutInfo.DeviceId.IsEmpty())
+        {
+            // Some platforms (such as Mac) don't use string identifiers so synthesize one from
+			// the device name and index
+            OutInfo.DeviceId = FString::Printf(TEXT("%s_%d"), *OutInfo.DeviceName, InputDeviceId);
+        }
 	}
 	catch (const std::exception& e)
 	{
