[AudioCaptureAudio]
* Fixed c++20 compile warning and removed C++17 requirement

diff --git a/Engine/Source/Runtime/AudioCaptureImplementations/AudioCaptureRtAudio/AudioCaptureRtAudio.Build.cs b/Engine/Source/Runtime/AudioCaptureImplementations/AudioCaptureRtAudio/AudioCaptureRtAudio.Build.cs
--- a/Engine/Source/Runtime/AudioCaptureImplementations/AudioCaptureRtAudio/AudioCaptureRtAudio.Build.cs
+++ b/Engine/Source/Runtime/AudioCaptureImplementations/AudioCaptureRtAudio/AudioCaptureRtAudio.Build.cs
@@ -8,9 +8,6 @@ public class AudioCaptureRtAudio : ModuleRules
 
 	public AudioCaptureRtAudio(ReadOnlyTargetRules Target) : base(Target)
 	{
-		// warning C5056: operator '!=': deprecated for array types
-		CppStandard = CppStandardVersion.Cpp17;
-
 		// RtAudio throws exceptions for errors.
 		bEnableExceptions = true;
 
diff --git a/Engine/Source/Runtime/AudioCaptureImplementations/AudioCaptureRtAudio/Private/RtAudio.cpp b/Engine/Source/Runtime/AudioCaptureImplementations/AudioCaptureRtAudio/Private/RtAudio.cpp
--- a/Engine/Source/Runtime/AudioCaptureImplementations/AudioCaptureRtAudio/Private/RtAudio.cpp
+++ b/Engine/Source/Runtime/AudioCaptureImplementations/AudioCaptureRtAudio/Private/RtAudio.cpp
@@ -4609,7 +4609,8 @@ bool RtApiWasapi::probeDeviceOpen( unsigned int device, StreamMode mode, unsigne
   // Set flags for buffer conversion.
   stream_.doConvertBuffer[mode] = false;
   if ( stream_.userFormat != stream_.deviceFormat[mode] ||
-       stream_.nUserChannels != stream_.nDeviceChannels )
+       stream_.nUserChannels[0] != stream_.nDeviceChannels[0] ||
+	   stream_.nUserChannels[1] != stream_.nDeviceChannels[1])
     stream_.doConvertBuffer[mode] = true;
   else if ( stream_.userInterleaved != stream_.deviceInterleaved[mode] &&
             stream_.nUserChannels[mode] > 1 )
