Fix for crash on exit in FWasapiStreamManager destructor

diff --git a/Engine/Source/Runtime/AudioCaptureImplementations/Windows/AudioCaputureWasapi/Private/AudioCaptureWasapiModule.cpp b/Engine/Source/Runtime/AudioCaptureImplementations/Windows/AudioCaputureWasapi/Private/AudioCaptureWasapiModule.cpp
--- a/Engine/Source/Runtime/AudioCaptureImplementations/Windows/AudioCaputureWasapi/Private/AudioCaptureWasapiModule.cpp
+++ b/Engine/Source/Runtime/AudioCaptureImplementations/Windows/AudioCaputureWasapi/Private/AudioCaptureWasapiModule.cpp
@@ -20,17 +20,27 @@ namespace Audio
 	class FAudioCaptureWasapiModule : public IModuleInterface
 	{
 	private:
+		/** Indicates if FWindowsPlatformMisc::CoInitialize() was successfull. */
+		bool bCoInitialized = false;
+		
 		FAudioCaptureWasapiFactory AudioCaptureFactory;
 
 	public:
 		virtual void StartupModule() override
 		{
+			bCoInitialized = FWindowsPlatformMisc::CoInitialize();
+			
 			IModularFeatures::Get().RegisterModularFeature(IAudioCaptureFactory::GetModularFeatureName(), &AudioCaptureFactory);
 		}
 
 		virtual void ShutdownModule() override
 		{
 			IModularFeatures::Get().UnregisterModularFeature(IAudioCaptureFactory::GetModularFeatureName(), &AudioCaptureFactory);
+			
+			if (bCoInitialized)
+			{
+				FWindowsPlatformMisc::CoUninitialize();
+			}
 		}
 	};
 }
diff --git a/Engine/Source/Runtime/AudioCaptureImplementations/Windows/AudioCaputureWasapi/Private/WasapiStreamManager.cpp b/Engine/Source/Runtime/AudioCaptureImplementations/Windows/AudioCaputureWasapi/Private/WasapiStreamManager.cpp
--- a/Engine/Source/Runtime/AudioCaptureImplementations/Windows/AudioCaputureWasapi/Private/WasapiStreamManager.cpp
+++ b/Engine/Source/Runtime/AudioCaptureImplementations/Windows/AudioCaputureWasapi/Private/WasapiStreamManager.cpp
@@ -10,23 +10,9 @@ namespace Audio
 {
 	FWasapiStreamManager::FWasapiStreamManager()
 	{
-		const bool bResult = FWindowsPlatformMisc::CoInitialize();
-		if (bResult)
-		{
-			bCoInitialized = true;
-		}
-
 		DeviceEnumerator.Initialize();
 	}
 
-	FWasapiStreamManager::~FWasapiStreamManager()
-	{
-		if (bCoInitialized)
-		{
-			FWindowsPlatformMisc::CoUninitialize();
-		}
-	}
-
 	FString FWasapiStreamManager::GetDefaultInputDeviceId()
 	{
 		return DeviceEnumerator.GetDefaultInputDeviceId();
diff --git a/Engine/Source/Runtime/AudioCaptureImplementations/Windows/AudioCaputureWasapi/Private/WasapiStreamManager.h b/Engine/Source/Runtime/AudioCaptureImplementations/Windows/AudioCaputureWasapi/Private/WasapiStreamManager.h
--- a/Engine/Source/Runtime/AudioCaptureImplementations/Windows/AudioCaputureWasapi/Private/WasapiStreamManager.h
+++ b/Engine/Source/Runtime/AudioCaptureImplementations/Windows/AudioCaputureWasapi/Private/WasapiStreamManager.h
@@ -25,7 +25,7 @@ namespace Audio
 		FWasapiStreamManager(FWasapiStreamManager&& InOther) = delete;
 		FWasapiStreamManager(const FWasapiStreamManager& InOther) = delete;
 
-		virtual ~FWasapiStreamManager();
+		virtual ~FWasapiStreamManager() {}
 
 		FWasapiStreamManager& operator=(FWasapiStreamManager&& InOther) = delete;
 		FWasapiStreamManager& operator=(const FWasapiStreamManager& InOther) = delete;
@@ -101,9 +101,6 @@ namespace Audio
 			STREAM_CAPTURING
 		};
 
-		/** Indicates if FWindowsPlatformMisc::CoInitialize() was successfull. */
-		bool bCoInitialized = false;
-
 		/** Current state of the state machine. */
 		EStreamState State = EStreamState::STREAM_CLOSED;
 
