Avoid using voice processing hardware on iOS if not requested when opening the stream since it can send sound to unexpected routes

diff --git a/Engine/Source/Runtime/AudioCaptureImplementations/IOS/AudioCaptureAudioUnit/Private/AudioCaptureAudioUnit.cpp b/Engine/Source/Runtime/AudioCaptureImplementations/IOS/AudioCaptureAudioUnit/Private/AudioCaptureAudioUnit.cpp
--- a/Engine/Source/Runtime/AudioCaptureImplementations/IOS/AudioCaptureAudioUnit/Private/AudioCaptureAudioUnit.cpp
+++ b/Engine/Source/Runtime/AudioCaptureImplementations/IOS/AudioCaptureAudioUnit/Private/AudioCaptureAudioUnit.cpp
@@ -1,15 +1,20 @@
 // Copyright Epic Games, Inc. All Rights Reserved.
 
 #include "AudioCaptureAudioUnit.h"
+#include "AudioCaptureCoreLog.h"
+#include <AVFoundation/AVAudioSession.h>
 
 const int32 kInputBus = 1;
 const int32 kOutputBus = 0;
+const int32 kRemoteIODeviceIndex = 0;
+const int32 kVoiceProcessingIODeviceIndex = 1;
 
 Audio::FAudioCaptureAudioUnitStream::FAudioCaptureAudioUnitStream()
 	: NumChannels(0)
 	, SampleRate(0)
 	, bIsStreamOpen(false)
 	, bHasCaptureStarted(false)
+	, bIsHardwareVoiceProcessingSupported(false)
 {
 }
 
@@ -76,32 +81,74 @@ void Audio::FAudioCaptureAudioUnitStream::AllocateBuffer(int SizeInBytes)
 
 bool Audio::FAudioCaptureAudioUnitStream::GetCaptureDeviceInfo(FCaptureDeviceInfo& OutInfo, int32 DeviceIndex)
 {
-	OutInfo.DeviceName = FString(TEXT("Remote IO Audio Component"));
-	OutInfo.InputChannels = 1;
-	OutInfo.PreferredSampleRate = 48000;
-
-	return true;
+	switch (DeviceIndex)
+	{
+		case Audio::DefaultDeviceIndex:
+		case kRemoteIODeviceIndex:
+		{
+			OutInfo.DeviceName = FString(TEXT("Remote IO Audio Component"));
+			OutInfo.InputChannels = 1;
+			OutInfo.PreferredSampleRate = 48000;
+			OutInfo.bSupportsHardwareAEC = false;
+			return true;
+		}
+		case kVoiceProcessingIODeviceIndex:
+		{
+			OutInfo.DeviceName = FString(TEXT("VoiceProcesing IO Audio Component"));
+			OutInfo.InputChannels = 1;
+			OutInfo.PreferredSampleRate = 48000;
+			OutInfo.bSupportsHardwareAEC = true;
+			return true;
+		}
+		default:
+		{
+			return false;
+		}
+	}
 }
 
 bool Audio::FAudioCaptureAudioUnitStream::OpenAudioCaptureStream(const FAudioCaptureDeviceParams& InParams, FOnAudioCaptureFunction InOnCapture, uint32 NumFramesDesired)
 {
+	switch (InParams.DeviceIndex)
+	{
+		case Audio::DefaultDeviceIndex:
+		case kRemoteIODeviceIndex:
+		{
+			bIsHardwareVoiceProcessingSupported	= false;
+			UE_CLOG(InParams.bUseHardwareAEC, LogAudioCaptureCore, Warning, TEXT("Hardware support is only available for VoiceProcessing IO Audio Component (DeviceIndex = %d)"), kVoiceProcessingIODeviceIndex);
+			break;
+		}
+		case kVoiceProcessingIODeviceIndex:
+		{
+			bIsHardwareVoiceProcessingSupported	= true;
+			break;
+		}
+		default:
+		{
+			return false;
+		}
+	}
+
 	NumChannels = 1;
 	SampleRate = 48000;
-	OSStatus Status = noErr;
-	
 	OnCapture = MoveTemp(InOnCapture);
-	
+
+	OSStatus Status = noErr;
+		
 	// Source of info "Technical Note TN2091 - Device input using the HAL Output Audio Unit"
 	
 	AudioComponentDescription desc;
 	desc.componentType = kAudioUnitType_Output;
-	// We use processing element always for enable runtime changing HW AEC and AGC settings.
-	// When it's disable the unit work as RemoteIO
-	desc.componentSubType = kAudioUnitSubType_VoiceProcessingIO;
+	desc.componentSubType = bIsHardwareVoiceProcessingSupported ? kAudioUnitSubType_VoiceProcessingIO : kAudioUnitSubType_RemoteIO;
 	desc.componentManufacturer = kAudioUnitManufacturer_Apple;
 	desc.componentFlags = 0;
 	desc.componentFlagsMask = 0;
 
+	// Using VoiceProcessing IO may change the AVAudioSession mode to AVAudioSessionModeVoiceChat. Cache AVAudioSession settings and restore them after initialisation
+	NSString* Mode = [[AVAudioSession sharedInstance] mode];
+	NSString* Category = [[AVAudioSession sharedInstance] category];
+	AVAudioSessionCategoryOptions Options = [[AVAudioSession sharedInstance] categoryOptions];
+
 	AudioComponent InputComponent = AudioComponentFindNext(NULL, &desc);
 
 	Status = AudioComponentInstanceNew(InputComponent, &IOUnit);
@@ -160,16 +207,27 @@ bool Audio::FAudioCaptureAudioUnitStream::OpenAudioCaptureStream(const FAudioCap
 		sizeof(CallbackInfo));
 	check(Status == noErr);
 	
+	// Configure unit processing
+	if (bIsHardwareVoiceProcessingSupported)
+	{
+		SetHardwareFeatureEnabled(Audio::EHardwareInputFeature::EchoCancellation, InParams.bUseHardwareAEC);
+		SetHardwareFeatureEnabled(Audio::EHardwareInputFeature::AutomaticGainControl, InParams.bUseHardwareAEC);
+	}
+
 	// Initialize audio unit
 	Status = AudioUnitInitialize(IOUnit);
 	check(Status == noErr);
 
-	// Configure unit processing
-	SetHardwareFeatureEnabled(Audio::EHardwareInputFeature::EchoCancellation, InParams.bUseHardwareAEC);
-	SetHardwareFeatureEnabled(Audio::EHardwareInputFeature::AutomaticGainControl, InParams.bUseHardwareAEC);
-
 	bIsStreamOpen = (Status == noErr);
 
+	if ([[[AVAudioSession sharedInstance] category] compare:Category] != NSOrderedSame ||
+		[[[AVAudioSession sharedInstance] mode] compare:Mode] != NSOrderedSame ||
+		[[AVAudioSession sharedInstance] categoryOptions] != Options)
+	{
+		NSError* ActiveError = nil;
+		[[AVAudioSession sharedInstance] setCategory:Category mode:Mode options:Options error:&ActiveError];
+	}
+
 	return bIsStreamOpen;
 }
 
@@ -224,19 +282,19 @@ void Audio::FAudioCaptureAudioUnitStream::OnAudioCapture(void* InBuffer, uint32
 
 bool Audio::FAudioCaptureAudioUnitStream::GetInputDevicesAvailable(TArray<FCaptureDeviceInfo>& OutDevices)
 {
-	// TODO: Add individual devices for different ports here.
 	OutDevices.Reset();
 
-	FCaptureDeviceInfo& DeviceInfo = OutDevices.AddDefaulted_GetRef();
-	GetCaptureDeviceInfo(DeviceInfo, 0);
+	GetCaptureDeviceInfo(OutDevices.AddDefaulted_GetRef(), kRemoteIODeviceIndex);
+	GetCaptureDeviceInfo(OutDevices.AddDefaulted_GetRef(), kVoiceProcessingIODeviceIndex);
 
 	return true;
 }
 
 void Audio::FAudioCaptureAudioUnitStream::SetHardwareFeatureEnabled(EHardwareInputFeature FeatureType, bool bEnabled)
 {
-	if (IOUnit == nil)
+	if (IOUnit == nil || !bIsHardwareVoiceProcessingSupported)
 	{
+		UE_CLOG(!bIsHardwareVoiceProcessingSupported, LogAudioCaptureCore, Warning, TEXT("Hardware support is only available for VoiceProcessing IO Audio Component (DeviceIndex = %d)"), kVoiceProcessingIODeviceIndex);
 		return;
 	}
 
diff --git a/Engine/Source/Runtime/AudioCaptureImplementations/IOS/AudioCaptureAudioUnit/Private/AudioCaptureAudioUnit.h b/Engine/Source/Runtime/AudioCaptureImplementations/IOS/AudioCaptureAudioUnit/Private/AudioCaptureAudioUnit.h
--- a/Engine/Source/Runtime/AudioCaptureImplementations/IOS/AudioCaptureAudioUnit/Private/AudioCaptureAudioUnit.h
+++ b/Engine/Source/Runtime/AudioCaptureImplementations/IOS/AudioCaptureAudioUnit/Private/AudioCaptureAudioUnit.h
@@ -40,5 +40,6 @@ namespace Audio
 		int BufferSize = 0;
 		bool bIsStreamOpen;
 		bool bHasCaptureStarted;
+		bool bIsHardwareVoiceProcessingSupported;
 	};
 }
