Static analysis fix

#rnx

diff --git a/Engine/Source/Runtime/AudioCaptureImplementations/AudioCaptureRtAudio/Private/RtAudio.cpp b/Engine/Source/Runtime/AudioCaptureImplementations/AudioCaptureRtAudio/Private/RtAudio.cpp
--- a/Engine/Source/Runtime/AudioCaptureImplementations/AudioCaptureRtAudio/Private/RtAudio.cpp
+++ b/Engine/Source/Runtime/AudioCaptureImplementations/AudioCaptureRtAudio/Private/RtAudio.cpp
@@ -4017,10 +4017,10 @@ RtAudio::DeviceInfo RtApiWasapi::getDeviceInfo( unsigned int device )
   unsigned int renderDeviceCount = 0;
   bool isCaptureDevice = false;
   std::string defaultDeviceId;
-  LPWSTR strDefaultDeviceId;
-  LPWSTR strDeviceId;
+  LPWSTR strDefaultDeviceId = NULL;
+  LPWSTR strDeviceId = NULL;
 
-  PROPVARIANT deviceNameProp;
+  PROPVARIANT deviceNameProp = {};
 
   IMMDeviceCollection* captureDevices = NULL;
   IMMDeviceCollection* renderDevices = NULL;
