PR #12092: Clear stale CaptureThread pointer in FWasapiCaptureThread::Stop & Abort

#rnx

diff --git a/Engine/Source/Runtime/AudioCaptureImplementations/Windows/AudioCaputureWasapi/Private/WasapiCaptureThread.cpp b/Engine/Source/Runtime/AudioCaptureImplementations/Windows/AudioCaputureWasapi/Private/WasapiCaptureThread.cpp
--- a/Engine/Source/Runtime/AudioCaptureImplementations/Windows/AudioCaputureWasapi/Private/WasapiCaptureThread.cpp
+++ b/Engine/Source/Runtime/AudioCaptureImplementations/Windows/AudioCaputureWasapi/Private/WasapiCaptureThread.cpp
@@ -78,6 +78,7 @@ namespace Audio
 
 			bool bShouldWait = true;
 			CaptureThread->Kill(bShouldWait);
+			CaptureThread = nullptr;
 		}
 	}
 
@@ -91,6 +92,7 @@ namespace Audio
 			// the stream is disposed of mid-callback.
 			bool bShouldWait = true;
 			CaptureThread->Kill(bShouldWait);
+			CaptureThread = nullptr;
 		}
 	}
 }
