[AudioCaptureImplementations] Disable AutoRTFM instrumentation.

AutoRTFM instrumentation cannot be enabled with exceptions.

diff --git a/Engine/Source/Runtime/AudioCaptureImplementations/AudioCaptureRtAudio/AudioCaptureRtAudio.Build.cs b/Engine/Source/Runtime/AudioCaptureImplementations/AudioCaptureRtAudio/AudioCaptureRtAudio.Build.cs
--- a/Engine/Source/Runtime/AudioCaptureImplementations/AudioCaptureRtAudio/AudioCaptureRtAudio.Build.cs
+++ b/Engine/Source/Runtime/AudioCaptureImplementations/AudioCaptureRtAudio/AudioCaptureRtAudio.Build.cs
@@ -10,6 +10,8 @@ public class AudioCaptureRtAudio : ModuleRules
 	{
 		// RtAudio throws exceptions for errors.
 		bEnableExceptions = true;
+		// AutoRTFM cannot be used with exceptions.
+		bDisableAutoRTFMInstrumentation = true;
 
 		PrivateDependencyModuleNames.Add("Core");
 		PrivateDependencyModuleNames.Add("AudioCaptureCore");
