Enabling explicit partform SDK requirement in multiple modules

This feature overall is disabled by default therefore it will not affect existing builds jobs, but setting the flag per module now allows us to iterate and test it more easily.

#tests Aside from the preflight I tested I can build with the feature enabled or disabled without issues LyraGame

#Synced-CL 50087102

diff --git a/Engine/Plugins/Media/AvfMedia/Source/AvfMedia/AvfMedia.Build.cs b/Engine/Plugins/Media/AvfMedia/Source/AvfMedia/AvfMedia.Build.cs
--- a/Engine/Plugins/Media/AvfMedia/Source/AvfMedia/AvfMedia.Build.cs
+++ b/Engine/Plugins/Media/AvfMedia/Source/AvfMedia/AvfMedia.Build.cs
@@ -6,6 +6,7 @@ namespace UnrealBuildTool.Rules
 	{
 		public AvfMedia(ReadOnlyTargetRules Target) : base(Target)
 		{
+			bRequiresPlatformSDK = true;
 			DynamicallyLoadedModuleNames.AddRange(
 				new string[] {
 					"Media",
diff --git a/Engine/Plugins/Media/AvfMedia/Source/AvfMediaCapture/AvfMediaCapture.Build.cs b/Engine/Plugins/Media/AvfMedia/Source/AvfMediaCapture/AvfMediaCapture.Build.cs
--- a/Engine/Plugins/Media/AvfMedia/Source/AvfMediaCapture/AvfMediaCapture.Build.cs
+++ b/Engine/Plugins/Media/AvfMedia/Source/AvfMediaCapture/AvfMediaCapture.Build.cs
@@ -6,6 +6,8 @@ namespace UnrealBuildTool.Rules
 	{
 		public AvfMediaCapture(ReadOnlyTargetRules Target) : base(Target)
 		{
+			bRequiresPlatformSDK = true;
+
 			DynamicallyLoadedModuleNames.AddRange(
 				new string[] {
 					"Media",
diff --git a/Engine/Plugins/Online/IOS/OnlineSubsystemIOS/Source/OnlineSubsystemIOS.Build.cs b/Engine/Plugins/Online/IOS/OnlineSubsystemIOS/Source/OnlineSubsystemIOS.Build.cs
--- a/Engine/Plugins/Online/IOS/OnlineSubsystemIOS/Source/OnlineSubsystemIOS.Build.cs
+++ b/Engine/Plugins/Online/IOS/OnlineSubsystemIOS/Source/OnlineSubsystemIOS.Build.cs
@@ -6,6 +6,8 @@ public class OnlineSubsystemIOS : ModuleRules
 {
     public OnlineSubsystemIOS(ReadOnlyTargetRules Target) : base(Target)
     {
+        bRequiresPlatformSDK = true;
+
 		PCHUsage = ModuleRules.PCHUsageMode.UseExplicitOrSharedPCHs;
 		PrivateIncludePaths.AddRange( 
             new string[] {
diff --git a/Engine/Plugins/Runtime/AppleImageUtils/Source/AppleImageUtils/AppleImageUtils.Build.cs b/Engine/Plugins/Runtime/AppleImageUtils/Source/AppleImageUtils/AppleImageUtils.Build.cs
--- a/Engine/Plugins/Runtime/AppleImageUtils/Source/AppleImageUtils/AppleImageUtils.Build.cs
+++ b/Engine/Plugins/Runtime/AppleImageUtils/Source/AppleImageUtils/AppleImageUtils.Build.cs
@@ -6,6 +6,7 @@ namespace UnrealBuildTool.Rules
 	{
 		public AppleImageUtils(ReadOnlyTargetRules Target) : base(Target)
 		{
+			bRequiresPlatformSDK = true;
             PublicDependencyModuleNames.AddRange(
                 new string[]
                 {
diff --git a/Engine/Plugins/Runtime/AppleMoviePlayer/Source/AppleMoviePlayer/AppleMoviePlayer.Build.cs b/Engine/Plugins/Runtime/AppleMoviePlayer/Source/AppleMoviePlayer/AppleMoviePlayer.Build.cs
--- a/Engine/Plugins/Runtime/AppleMoviePlayer/Source/AppleMoviePlayer/AppleMoviePlayer.Build.cs
+++ b/Engine/Plugins/Runtime/AppleMoviePlayer/Source/AppleMoviePlayer/AppleMoviePlayer.Build.cs
@@ -6,6 +6,7 @@ namespace UnrealBuildTool.Rules
 	{
         public AppleMoviePlayer(ReadOnlyTargetRules Target) : base(Target)
 		{
+			bRequiresPlatformSDK = true;
 			PublicDependencyModuleNames.AddRange(
 				new string[] {
 					"Core",
diff --git a/Engine/Source/Runtime/AVIWriter/AVIWriter.Build.cs b/Engine/Source/Runtime/AVIWriter/AVIWriter.Build.cs
--- a/Engine/Source/Runtime/AVIWriter/AVIWriter.Build.cs
+++ b/Engine/Source/Runtime/AVIWriter/AVIWriter.Build.cs
@@ -6,6 +6,8 @@ namespace UnrealBuildTool.Rules
 	{
 		public AVIWriter(ReadOnlyTargetRules Target) : base(Target)
 		{
+			bRequiresPlatformSDK = true;
+
 			PublicDependencyModuleNames.AddRange(
 				new string[]
 				{
diff --git a/Engine/Source/Runtime/Apple/AudioMixerAudioUnit/AudioMixerAudioUnit.Build.cs b/Engine/Source/Runtime/Apple/AudioMixerAudioUnit/AudioMixerAudioUnit.Build.cs
--- a/Engine/Source/Runtime/Apple/AudioMixerAudioUnit/AudioMixerAudioUnit.Build.cs
+++ b/Engine/Source/Runtime/Apple/AudioMixerAudioUnit/AudioMixerAudioUnit.Build.cs
@@ -6,6 +6,7 @@ public class AudioMixerAudioUnit : ModuleRules
 {
 	public AudioMixerAudioUnit(ReadOnlyTargetRules Target) : base(Target)
 	{
+		bRequiresPlatformSDK = true;
 		PrivateIncludePathModuleNames.Add("TargetPlatform");
 		PublicIncludePaths.Add("Runtime/AudioMixer/Public");
 		PrivateIncludePaths.Add("Runtime/AudioMixer/Private");
diff --git a/Engine/Source/Runtime/Apple/MetalRHI/MetalRHI.Build.cs b/Engine/Source/Runtime/Apple/MetalRHI/MetalRHI.Build.cs
--- a/Engine/Source/Runtime/Apple/MetalRHI/MetalRHI.Build.cs
+++ b/Engine/Source/Runtime/Apple/MetalRHI/MetalRHI.Build.cs
@@ -7,6 +7,8 @@ public class MetalRHI : ModuleRules
 {	
 	public MetalRHI(ReadOnlyTargetRules Target) : base(Target)
 	{
+		bRequiresPlatformSDK = true;
+
 		PrivateDependencyModuleNames.AddRange(
 			new string[] {
 				"Core",
diff --git a/Engine/Source/Runtime/ApplicationCore/ApplicationCore.Build.cs b/Engine/Source/Runtime/ApplicationCore/ApplicationCore.Build.cs
--- a/Engine/Source/Runtime/ApplicationCore/ApplicationCore.Build.cs
+++ b/Engine/Source/Runtime/ApplicationCore/ApplicationCore.Build.cs
@@ -6,6 +6,7 @@ public class ApplicationCore : ModuleRules
 {
 	public ApplicationCore(ReadOnlyTargetRules Target) : base(Target)
 	{
+		bRequiresPlatformSDK = true;
 		PublicDependencyModuleNames.AddRange(
 			new string[] {
 				"Core"
diff --git a/Engine/Source/Runtime/AudioCaptureImplementations/IOS/AudioCaptureAudioUnit/AudioCaptureAudioUnit.Build.cs b/Engine/Source/Runtime/AudioCaptureImplementations/IOS/AudioCaptureAudioUnit/AudioCaptureAudioUnit.Build.cs
--- a/Engine/Source/Runtime/AudioCaptureImplementations/IOS/AudioCaptureAudioUnit/AudioCaptureAudioUnit.Build.cs
+++ b/Engine/Source/Runtime/AudioCaptureImplementations/IOS/AudioCaptureAudioUnit/AudioCaptureAudioUnit.Build.cs
@@ -6,6 +6,7 @@ public class AudioCaptureAudioUnit : ModuleRules
 {
     public AudioCaptureAudioUnit(ReadOnlyTargetRules Target) : base(Target)
     {
+	    bRequiresPlatformSDK = true;
         PublicFrameworks.AddRange(new string[] { "CoreAudio", "AVFoundation", "AudioToolbox" });
         PrivateDependencyModuleNames.Add("Core");
         PrivateDependencyModuleNames.Add("AudioCaptureCore");
diff --git a/Engine/Source/Runtime/Core/Core.Build.cs b/Engine/Source/Runtime/Core/Core.Build.cs
--- a/Engine/Source/Runtime/Core/Core.Build.cs
+++ b/Engine/Source/Runtime/Core/Core.Build.cs
@@ -8,6 +8,8 @@ public class Core : ModuleRules
 {
 	public Core(ReadOnlyTargetRules Target) : base(Target)
 	{
+		bRequiresPlatformSDK = true;
+
 		NumIncludedBytesPerUnityCPPOverride = 491520; // best unity size found from using UBT ProfileUnitySizes mode
 
 		PrivatePCHHeaderFile = "Private/CorePrivatePCH.h";
diff --git a/Engine/Source/Runtime/IOS/IOSAudio/IOSAudio.Build.cs b/Engine/Source/Runtime/IOS/IOSAudio/IOSAudio.Build.cs
--- a/Engine/Source/Runtime/IOS/IOSAudio/IOSAudio.Build.cs
+++ b/Engine/Source/Runtime/IOS/IOSAudio/IOSAudio.Build.cs
@@ -6,6 +6,8 @@ public class IOSAudio : ModuleRules
 {
 	public IOSAudio(ReadOnlyTargetRules Target) : base(Target)
 	{
+		bRequiresPlatformSDK = true;
+
 		BinariesSubFolder = "IOS";
 
 		PrivateIncludePathModuleNames.Add("TargetPlatform");
diff --git a/Engine/Source/Runtime/IOS/IOSLocalNotification/IOSLocalNotification.Build.cs b/Engine/Source/Runtime/IOS/IOSLocalNotification/IOSLocalNotification.Build.cs
--- a/Engine/Source/Runtime/IOS/IOSLocalNotification/IOSLocalNotification.Build.cs
+++ b/Engine/Source/Runtime/IOS/IOSLocalNotification/IOSLocalNotification.Build.cs
@@ -6,6 +6,8 @@ public class IOSLocalNotification : ModuleRules
 {
 	public IOSLocalNotification(ReadOnlyTargetRules Target) : base(Target)
 	{
+		bRequiresPlatformSDK = true;
+
 		BinariesSubFolder = "IOS";
 
 		PublicIncludePaths.AddRange(new string[]
diff --git a/Engine/Source/Runtime/IOS/IOSPlatformFeatures/IOSPlatformFeatures.Build.cs b/Engine/Source/Runtime/IOS/IOSPlatformFeatures/IOSPlatformFeatures.Build.cs
--- a/Engine/Source/Runtime/IOS/IOSPlatformFeatures/IOSPlatformFeatures.Build.cs
+++ b/Engine/Source/Runtime/IOS/IOSPlatformFeatures/IOSPlatformFeatures.Build.cs
@@ -6,6 +6,8 @@ public class IOSPlatformFeatures : ModuleRules
 {
 	public IOSPlatformFeatures(ReadOnlyTargetRules Target) : base(Target)
 	{
+		bRequiresPlatformSDK = true;
+
 		PrivateDependencyModuleNames.AddRange(new string[] 
 			{ 
 				"Core", 
diff --git a/Engine/Source/Runtime/IOS/MarketplaceKit/MarketplaceKit.Build.cs b/Engine/Source/Runtime/IOS/MarketplaceKit/MarketplaceKit.Build.cs
--- a/Engine/Source/Runtime/IOS/MarketplaceKit/MarketplaceKit.Build.cs
+++ b/Engine/Source/Runtime/IOS/MarketplaceKit/MarketplaceKit.Build.cs
@@ -7,6 +7,8 @@ public class MarketplaceKit : ModuleRules
 {
 	public MarketplaceKit(ReadOnlyTargetRules Target) : base(Target)
 	{
+		bRequiresPlatformSDK = true;
+
 		PrivateDependencyModuleNames.Add("Swift");
 
 		PrivateDependencyModuleNames.AddRange(new string[]
diff --git a/Engine/Source/Runtime/ImageWrapper/ImageWrapper.Build.cs b/Engine/Source/Runtime/ImageWrapper/ImageWrapper.Build.cs
--- a/Engine/Source/Runtime/ImageWrapper/ImageWrapper.Build.cs
+++ b/Engine/Source/Runtime/ImageWrapper/ImageWrapper.Build.cs
@@ -8,6 +8,8 @@ public class ImageWrapper : ModuleRules
 
 	public ImageWrapper(ReadOnlyTargetRules Target) : base(Target)
 	{
+		bRequiresPlatformSDK = true;
+
 		PublicIncludePathModuleNames.AddRange(new string[] {
 			"ImageCore"
 			}
diff --git a/Engine/Source/Runtime/InputCore/InputCore.Build.cs b/Engine/Source/Runtime/InputCore/InputCore.Build.cs
--- a/Engine/Source/Runtime/InputCore/InputCore.Build.cs
+++ b/Engine/Source/Runtime/InputCore/InputCore.Build.cs
@@ -6,6 +6,8 @@ public class InputCore : ModuleRules
 {
 	public InputCore(ReadOnlyTargetRules Target) : base(Target)
 	{
+		bRequiresPlatformSDK = true;
+
 		PublicDependencyModuleNames.AddRange(new string[] { "Core", "CoreUObject" });
 
 		if(Target.IsInPlatformGroup(UnrealPlatformGroup.IOS))
diff --git a/Engine/Source/Runtime/Launch/Launch.Build.cs b/Engine/Source/Runtime/Launch/Launch.Build.cs
--- a/Engine/Source/Runtime/Launch/Launch.Build.cs
+++ b/Engine/Source/Runtime/Launch/Launch.Build.cs
@@ -24,6 +24,8 @@ public class Launch : ModuleRules
 	
 	public Launch(ReadOnlyTargetRules Target) : base(Target)
 	{
+		bRequiresPlatformSDK = true;
+
 		PrivateIncludePathModuleNames.AddRange(new string[] {
 				"AutomationController",
 				"AutomationTest",
diff --git a/Engine/Source/Runtime/Online/HTTP/HTTP.Build.cs b/Engine/Source/Runtime/Online/HTTP/HTTP.Build.cs
--- a/Engine/Source/Runtime/Online/HTTP/HTTP.Build.cs
+++ b/Engine/Source/Runtime/Online/HTTP/HTTP.Build.cs
@@ -97,6 +97,8 @@ public class HTTP : ModuleRules
 
 	public HTTP(ReadOnlyTargetRules Target) : base(Target)
 	{
+		bRequiresPlatformSDK = true;
+
 		PublicDefinitions.Add("HTTP_PACKAGE=1");
 
 		PublicDependencyModuleNames.AddRange(
diff --git a/Engine/Source/Runtime/PreLoadScreen/PreLoadScreen.Build.cs b/Engine/Source/Runtime/PreLoadScreen/PreLoadScreen.Build.cs
--- a/Engine/Source/Runtime/PreLoadScreen/PreLoadScreen.Build.cs
+++ b/Engine/Source/Runtime/PreLoadScreen/PreLoadScreen.Build.cs
@@ -6,6 +6,8 @@ public class PreLoadScreen : ModuleRules
 {
 	public PreLoadScreen(ReadOnlyTargetRules Target) : base(Target)
 	{
+		bRequiresPlatformSDK = true;
+
 		PublicDependencyModuleNames.AddRange(
 			new string[] {
 					"Engine",
diff --git a/Engine/Source/Runtime/RHI/RHI.Build.cs b/Engine/Source/Runtime/RHI/RHI.Build.cs
--- a/Engine/Source/Runtime/RHI/RHI.Build.cs
+++ b/Engine/Source/Runtime/RHI/RHI.Build.cs
@@ -8,6 +8,8 @@ public class RHI : ModuleRules
 {
 	public RHI(ReadOnlyTargetRules Target) : base(Target)
 	{
+		bRequiresPlatformSDK = true;
+
 		PrivateDependencyModuleNames.Add("Core");
 		PrivateDependencyModuleNames.Add("TraceLog");
 		PrivateDependencyModuleNames.Add("ApplicationCore");
diff --git a/Engine/Source/Runtime/Slate/Slate.Build.cs b/Engine/Source/Runtime/Slate/Slate.Build.cs
--- a/Engine/Source/Runtime/Slate/Slate.Build.cs
+++ b/Engine/Source/Runtime/Slate/Slate.Build.cs
@@ -6,6 +6,8 @@ public class Slate : ModuleRules
 {
 	public Slate(ReadOnlyTargetRules Target) : base(Target)
 	{
+		bRequiresPlatformSDK = true;
+
         PrivateDefinitions.Add("SLATE_MODULE=1");
         SharedPCHHeaderFile = "Public/SlateSharedPCH.h";
 
diff --git a/Engine/Source/Runtime/WebBrowser/WebBrowser.Build.cs b/Engine/Source/Runtime/WebBrowser/WebBrowser.Build.cs
--- a/Engine/Source/Runtime/WebBrowser/WebBrowser.Build.cs
+++ b/Engine/Source/Runtime/WebBrowser/WebBrowser.Build.cs
@@ -7,6 +7,8 @@ public class WebBrowser : ModuleRules
 {
 	public WebBrowser(ReadOnlyTargetRules Target) : base(Target)
 	{
+		bRequiresPlatformSDK = true;
+
 		PrivateDependencyModuleNames.AddRange(
 			new string[]
 			{
