Remove unused variable and free of null
#rnx

diff --git a/Engine/Source/Runtime/AudioCaptureImplementations/AudioCaptureRtAudio/Private/RtAudio.cpp b/Engine/Source/Runtime/AudioCaptureImplementations/AudioCaptureRtAudio/Private/RtAudio.cpp
--- a/Engine/Source/Runtime/AudioCaptureImplementations/AudioCaptureRtAudio/Private/RtAudio.cpp
+++ b/Engine/Source/Runtime/AudioCaptureImplementations/AudioCaptureRtAudio/Private/RtAudio.cpp
@@ -4031,7 +4031,6 @@ RtAudio::DeviceInfo RtApiWasapi::getDeviceInfo( unsigned int device )
   IPropertyStore* defaultDevicePropStore = NULL;
 
   WAVEFORMATEX* deviceFormat = NULL;
-  WAVEFORMATEX* closestMatchFormat = NULL;
 
   // probed
   info.probed = false;
@@ -4234,7 +4233,6 @@ Exit:
   SAFE_RELEASE( defaultDevicePropStore );
 
   CoTaskMemFree( deviceFormat );
-  CoTaskMemFree( closestMatchFormat );
 
   if ( !errorText_.empty() )
     error( errorType );
